CIS Fixes (Script plugin related)

diff --git a/Engine/Plugins/Script/ScriptPlugin/Source/ScriptPlugin/Classes/ScriptComponent.h b/Engine/Plugins/Script/ScriptPlugin/Source/ScriptPlugin/Classes/ScriptComponent.h
--- a/Engine/Plugins/Script/ScriptPlugin/Source/ScriptPlugin/Classes/ScriptComponent.h
+++ b/Engine/Plugins/Script/ScriptPlugin/Source/ScriptPlugin/Classes/ScriptComponent.h
@@ -12,6 +12,7 @@ public:
 	virtual void Tick(float DeltaTime) = 0;
 	virtual void Destroy() = 0;
 	virtual bool CanTick() = 0;
+	virtual ~FScriptContextBase() {}
 };
 
 /** Component that allows you to specify custom triangle mesh geometry */
