Don't tick ScriptComponent in editor (also don't attempt to initialize scripts in non-game worlds)

diff --git a/Engine/Plugins/Script/ScriptPlugin/Source/ScriptPlugin/Private/ScriptComponent.cpp b/Engine/Plugins/Script/ScriptPlugin/Source/ScriptPlugin/Private/ScriptComponent.cpp
--- a/Engine/Plugins/Script/ScriptPlugin/Source/ScriptPlugin/Private/ScriptComponent.cpp
+++ b/Engine/Plugins/Script/ScriptPlugin/Source/ScriptPlugin/Private/ScriptComponent.cpp
@@ -8,7 +8,7 @@ UScriptComponent::UScriptComponent(const FPostConstructInitializeProperties& PCI
 	: Super( PCIP )
 {	
 	PrimaryComponentTick.bCanEverTick = true;
-	bTickInEditor = true;
+	bTickInEditor = false;
 	bAutoActivate = true;
 
 	Context = NULL;
@@ -18,7 +18,7 @@ UScriptComponent::UScriptComponent(const FPostConstructInitializeProperties& PCI
 void UScriptComponent::OnRegister()
 {
 	Super::OnRegister();
-	if (Script)
+	if (Script && GetWorld() && GetWorld()->WorldType != EWorldType::Editor)
 	{
 		// @todo Create context here
 #if WITH_LUA
@@ -32,8 +32,7 @@ void UScriptComponent::OnRegister()
 				bDoNotTick = !Context->CanTick();
 			}
 			if (bDoNotTick)
-			{
-				bTickInEditor = false;
+			{				
 				bAutoActivate = false;
 				PrimaryComponentTick.bCanEverTick = false;
 			}
