ScriptPlugin - Added missing include

diff --git a/Engine/Plugins/Script/ScriptPlugin/Source/ScriptPlugin/Private/ScriptPluginPrivatePCH.h b/Engine/Plugins/Script/ScriptPlugin/Source/ScriptPlugin/Private/ScriptPluginPrivatePCH.h
--- a/Engine/Plugins/Script/ScriptPlugin/Source/ScriptPlugin/Private/ScriptPluginPrivatePCH.h
+++ b/Engine/Plugins/Script/ScriptPlugin/Source/ScriptPlugin/Private/ScriptPluginPrivatePCH.h
@@ -7,6 +7,7 @@
 #include "Engine.h"
 #include "ScriptComponent.h"
 #include "ScriptAsset.h"
+#include "ScriptTestActor.h"
 
 #if WITH_LUA
 #include "LuaIntegration.h"
