Disable ScriptPlugin for Linux temporarily.

This is needed until we can build UHT on Linux with plugins. So far UBT does not support non-monolithics builds on that platform and having this plugin enabled breaks native compilation.

diff --git a/Engine/Plugins/Script/ScriptPlugin/ScriptPlugin.uplugin b/Engine/Plugins/Script/ScriptPlugin/ScriptPlugin.uplugin
--- a/Engine/Plugins/Script/ScriptPlugin/ScriptPlugin.uplugin
+++ b/Engine/Plugins/Script/ScriptPlugin/ScriptPlugin.uplugin
@@ -16,6 +16,7 @@
 			"Name" : "ScriptPlugin",
 			"Type" : "Runtime",
 			"LoadingPhase" : "PreDefault"
+			"BlacklistPlatforms" : [ "Linux" ]
 		}
 	]
-}
\ No newline at end of file
+}
