Fixing missing comma in ScriptPlugin.uplugin

diff --git a/Engine/Plugins/Script/ScriptPlugin/ScriptPlugin.uplugin b/Engine/Plugins/Script/ScriptPlugin/ScriptPlugin.uplugin
--- a/Engine/Plugins/Script/ScriptPlugin/ScriptPlugin.uplugin
+++ b/Engine/Plugins/Script/ScriptPlugin/ScriptPlugin.uplugin
@@ -15,7 +15,7 @@
 		{
 			"Name" : "ScriptPlugin",
 			"Type" : "Runtime",
-			"LoadingPhase" : "PreDefault"
+			"LoadingPhase" : "PreDefault",
 			"BlacklistPlatforms" : [ "Linux" ]
 		}
 	]
