ScriptPlugin: adding module API

diff --git a/Engine/Plugins/Script/ScriptPlugin/Source/ScriptPlugin/Classes/ScriptAsset.h b/Engine/Plugins/Script/ScriptPlugin/Source/ScriptPlugin/Classes/ScriptAsset.h
--- a/Engine/Plugins/Script/ScriptPlugin/Source/ScriptPlugin/Classes/ScriptAsset.h
+++ b/Engine/Plugins/Script/ScriptPlugin/Source/ScriptPlugin/Classes/ScriptAsset.h
@@ -4,7 +4,7 @@
 
 /** Definition of a specific ability that is applied to a character. Exists as part of a Trait. */
 UCLASS(BlueprintType, MinimalAPI)
-class UScriptAsset : public UDataAsset
+class SCRIPTPLUGIN_API UScriptAsset : public UDataAsset
 {
 	GENERATED_UCLASS_BODY()
 
diff --git a/Engine/Plugins/Script/ScriptPlugin/Source/ScriptPlugin/Classes/ScriptComponent.h b/Engine/Plugins/Script/ScriptPlugin/Source/ScriptPlugin/Classes/ScriptComponent.h
--- a/Engine/Plugins/Script/ScriptPlugin/Source/ScriptPlugin/Classes/ScriptComponent.h
+++ b/Engine/Plugins/Script/ScriptPlugin/Source/ScriptPlugin/Classes/ScriptComponent.h
@@ -18,7 +18,7 @@ public:
 
 /** Component that allows you to specify custom triangle mesh geometry */
 UCLASS(hidecategories=(Object, ActorComponent), editinlinenew, meta=(BlueprintSpawnableComponent), ClassGroup=Script)
-class UScriptComponent : public UActorComponent
+class SCRIPTPLUGIN_API UScriptComponent : public UActorComponent
 {
 	GENERATED_UCLASS_BODY()
 
