ScriptPlugin: removing MinimalAPI from ScriptAsset (already has module API)

diff --git a/Engine/Plugins/Script/ScriptPlugin/Source/ScriptPlugin/Classes/ScriptAsset.h b/Engine/Plugins/Script/ScriptPlugin/Source/ScriptPlugin/Classes/ScriptAsset.h
--- a/Engine/Plugins/Script/ScriptPlugin/Source/ScriptPlugin/Classes/ScriptAsset.h
+++ b/Engine/Plugins/Script/ScriptPlugin/Source/ScriptPlugin/Classes/ScriptAsset.h
@@ -3,7 +3,7 @@
 #include "ScriptAsset.generated.h"
 
 /** Definition of a specific ability that is applied to a character. Exists as part of a Trait. */
-UCLASS(BlueprintType, MinimalAPI)
+UCLASS(BlueprintType)
 class SCRIPTPLUGIN_API UScriptAsset : public UDataAsset
 {
 	GENERATED_UCLASS_BODY()
