ScriptPlugin: Script Generated Blueprint classes.

#ue4
#change Exposed the entire UFunction outside of CoreUObject
#change Script stack frame will now track currently executed native UFunctions
#change Added a few more lua types (FLinearColor etc)
#change Removed ScriptAsset
#change Added ScriptBlueprint and ScriptGeneratedBlueprintClass
#change Added ScriptActor. Both ScriptActor and ScriptComponent now serve the purpose of base classes for script generated classes
#change ScriptGeneratedBlueprint classes have properties and functions exported from script

diff --git a/Engine/Plugins/Script/ScriptEditorPlugin/Source/ScriptEditorPlugin/Classes/ReimportScriptFactory.h b/Engine/Plugins/Script/ScriptEditorPlugin/Source/ScriptEditorPlugin/Classes/ReimportScriptFactory.h
--- a/Engine/Plugins/Script/ScriptEditorPlugin/Source/ScriptEditorPlugin/Classes/ReimportScriptFactory.h
+++ b/Engine/Plugins/Script/ScriptEditorPlugin/Source/ScriptEditorPlugin/Classes/ReimportScriptFactory.h
@@ -8,13 +8,18 @@
 #include "ScriptFactory.h"
 #include "ReimportScriptFactory.generated.h"
 
-UCLASS(MinimalApi, collapsecategories)
+/**
+* Script Blueprint re-import factory
+*/
+UCLASS(MinimalApi, collapsecategories, EarlyAccessPreview)
 class UReimportScriptFactory : public UScriptFactory, public FReimportHandler
 {
 	GENERATED_UCLASS_BODY()
 
 	UPROPERTY()
-	class UScriptAsset* OriginalScript;
+	class UScriptBlueprint* OriginalScript;
+
+	virtual bool ConfigureProperties() override;
 
 	// Begin FReimportHandler interface
 	virtual bool CanReimport(UObject* Obj, TArray<FString>& OutFilenames) override;
diff --git a/Engine/Plugins/Script/ScriptEditorPlugin/Source/ScriptEditorPlugin/Classes/ScriptFactory.h b/Engine/Plugins/Script/ScriptEditorPlugin/Source/ScriptEditorPlugin/Classes/ScriptFactory.h
--- a/Engine/Plugins/Script/ScriptEditorPlugin/Source/ScriptEditorPlugin/Classes/ScriptFactory.h
+++ b/Engine/Plugins/Script/ScriptEditorPlugin/Source/ScriptEditorPlugin/Classes/ScriptFactory.h
@@ -2,20 +2,28 @@
 
 #pragma once
 
+#include "ScriptBlueprint.h"
 #include "ScriptFactory.generated.h"
 
-UCLASS(collapsecategories, hidecategories=Object)
+/**
+* Script Blueprint factory
+*/
+UCLASS(collapsecategories, hidecategories = Object, EarlyAccessPreview)
 class SCRIPTEDITORPLUGIN_API UScriptFactory : public UFactory
 {
 	GENERATED_UCLASS_BODY()
 
 public:
 
+	UPROPERTY(EditAnywhere, Category = ScriptFactory, meta = (AllowAbstract = "", BlueprintBaseOnly = ""))
+	TSubclassOf<class UObject> ParentClass;
+
 	// Begin UObject Interface
 	virtual void PostInitProperties() override;
 	// End UObject Interface
 
 	// Begin UFactory Interface
+	virtual bool ConfigureProperties() override;
 	virtual bool DoesSupportClass(UClass* Class) override;
 	virtual UObject* FactoryCreateText(UClass* InClass, UObject* InParent, FName InName, EObjectFlags Flags, UObject* Context, const TCHAR* Type, const TCHAR*& Buffer, const TCHAR* BufferEnd, FFeedbackContext* Warn);
 	// End UFactory Interface
diff --git a/Engine/Plugins/Script/ScriptEditorPlugin/Source/ScriptEditorPlugin/Private/ScriptBlueprintCompiler.cpp b/Engine/Plugins/Script/ScriptEditorPlugin/Source/ScriptEditorPlugin/Private/ScriptBlueprintCompiler.cpp
new file mode 100644
--- /dev/null
+++ b/Engine/Plugins/Script/ScriptEditorPlugin/Source/ScriptEditorPlugin/Private/ScriptBlueprintCompiler.cpp
@@ -0,0 +1,163 @@
+// Copyright 1998-2014 Epic Games, Inc. All Rights Reserved.
+
+#include "ScriptEditorPluginPrivatePCH.h"
+#include "ScriptBlueprint.h"
+#include "ScriptBlueprintGeneratedClass.h"
+#include "ScriptBlueprintCompiler.h"
+#include "Kismet2NameValidators.h"
+#include "KismetReinstanceUtilities.h"
+#include "ScriptFunction.h"
+
+///-------------------------------------------------------------
+
+FScriptBlueprintCompiler::FScriptBlueprintCompiler(UScriptBlueprint* SourceSketch, FCompilerResultsLog& InMessageLog, const FKismetCompilerOptions& InCompilerOptions, TArray<UObject*>* InObjLoaded)
+	: Super(SourceSketch, InMessageLog, InCompilerOptions, InObjLoaded)
+{
+}
+
+FScriptBlueprintCompiler::~FScriptBlueprintCompiler()
+{
+}
+
+void FScriptBlueprintCompiler::CopyTermDefaultsToDefaultObject(UObject* DefaultObject)
+{
+	Super::CopyTermDefaultsToDefaultObject(DefaultObject);
+}
+
+void FScriptBlueprintCompiler::CleanAndSanitizeClass(UBlueprintGeneratedClass* ClassToClean, UObject*& OldCDO)
+{
+	Super::CleanAndSanitizeClass(ClassToClean, OldCDO);
+
+	// Make sure our typed pointer is set
+	check(ClassToClean == NewClass);
+	NewScriptBlueprintClass = CastChecked<UScriptBlueprintGeneratedClass>((UObject*)NewClass);
+}
+
+void FScriptBlueprintCompiler::CreateClassVariablesFromBlueprint()
+{
+	Super::CreateClassVariablesFromBlueprint();
+
+	UScriptBlueprint* Blueprint = ScriptBlueprint();
+
+	for (auto& Field : ScriptDefinedFields)
+	{
+		if (Field.Class->IsChildOf(UFunction::StaticClass()))
+		{
+			// Functions are currently only supported for AScriptActors
+			if (Blueprint->ParentClass->IsChildOf(AScriptActor::StaticClass()))
+			{
+				UFunction* NewFunction = new(NewClass, Field.Name, RF_Public) UScriptFunction(FPostConstructInitializeProperties());
+				NewFunction->Bind();
+				NewFunction->StaticLink();
+				NewClass->AddFunctionToFunctionMap(NewFunction);
+				NewFunction->Next = NewClass->Children;
+				NewClass->Children = NewFunction;
+			}
+		}
+		else if (Field.Class->IsChildOf(UProperty::StaticClass()))
+		{
+			FString PinCategory;
+			if (Field.Class->IsChildOf(UStrProperty::StaticClass()))
+			{
+				PinCategory = Schema->PC_String;
+			}
+			else if (Field.Class->IsChildOf(UFloatProperty::StaticClass()))
+			{
+				PinCategory = Schema->PC_Float;
+			}
+			else if (Field.Class->IsChildOf(UIntProperty::StaticClass()))
+			{
+				PinCategory = Schema->PC_Int;
+			}
+			else if (Field.Class->IsChildOf(UBoolProperty::StaticClass()))
+			{
+				PinCategory = Schema->PC_Boolean;
+			}
+			else if (Field.Class->IsChildOf(UObjectProperty::StaticClass()))
+			{
+				PinCategory = Schema->PC_Object;
+			}
+			if (!PinCategory.IsEmpty())
+			{			
+				FEdGraphPinType ScriptPinType(PinCategory, TEXT(""), Field.Class, false, false);
+				UProperty* ScriptProperty = CreateVariable(Field.Name, ScriptPinType);
+				if (ScriptProperty != NULL)
+				{
+					ScriptProperty->SetMetaData(TEXT("Category"), *Blueprint->GetName());
+					ScriptProperty->SetPropertyFlags(CPF_BlueprintVisible | CPF_Edit);
+				}
+			}
+		}
+	}
+}
+
+void FScriptBlueprintCompiler::FinishCompilingClass(UClass* Class)
+{
+	UScriptBlueprint* Blueprint = ScriptBlueprint();
+
+	UScriptBlueprintGeneratedClass* ScriptClass = CastChecked<UScriptBlueprintGeneratedClass>(Class);
+	ScriptClass->SourceCode = Blueprint->SourceCode;
+	ScriptClass->ByteCode = Blueprint->ByteCode;
+
+	Super::FinishCompilingClass(Class);
+
+	// Allow Blueprint Components to be used in Blueprints
+	if (Blueprint->ParentClass->IsChildOf(UScriptComponent::StaticClass()) && Class != Blueprint->SkeletonGeneratedClass)
+	{
+		Class->SetMetaData(TEXT("BlueprintSpawnableComponent"), TEXT("true"));
+	}
+}
+
+void FScriptBlueprintCompiler::Compile()
+{
+	ScriptContext = NewScriptBlueprintClass->CreateContext();
+	bool Result = true;
+	if (ScriptContext.IsValid())
+	{
+		ScriptBlueprint()->UpdateSourceCodeIfChanged();
+		Result = ScriptContext->Initialize(ScriptBlueprint()->SourceCode, NULL);
+		if (!Result)
+		{
+			ScriptContext.Reset();
+		}
+	}
+	if (Result)
+	{
+		ScriptDefinedFields.Empty();
+		ScriptContext->GetScriptDefinedFields(ScriptDefinedFields);
+	}
+
+	Super::Compile();
+}
+
+void FScriptBlueprintCompiler::EnsureProperGeneratedClass(UClass*& TargetUClass)
+{
+	if ( TargetUClass && !( (UObject*)TargetUClass )->IsA(UScriptBlueprintGeneratedClass::StaticClass()) )
+	{
+		FKismetCompilerUtilities::ConsignToOblivion(TargetUClass, Blueprint->bIsRegeneratingOnLoad);
+		TargetUClass = NULL;
+	}
+}
+
+void FScriptBlueprintCompiler::SpawnNewClass(const FString& NewClassName)
+{
+	NewScriptBlueprintClass = FindObject<UScriptBlueprintGeneratedClass>(Blueprint->GetOutermost(), *NewClassName);
+
+	if ( NewScriptBlueprintClass == NULL )
+	{
+		NewScriptBlueprintClass = ConstructObject<UScriptBlueprintGeneratedClass>(UScriptBlueprintGeneratedClass::StaticClass(), Blueprint->GetOutermost(), FName(*NewClassName), RF_Public | RF_Transactional);
+	}
+	else
+	{
+		// Already existed, but wasn't linked in the Blueprint yet due to load ordering issues
+		FBlueprintCompileReinstancer GeneratedClassReinstancer(NewScriptBlueprintClass);
+	}
+	NewClass = NewScriptBlueprintClass;
+}
+
+bool FScriptBlueprintCompiler::ValidateGeneratedClass(UBlueprintGeneratedClass* Class)
+{
+	bool SuperResult = Super::ValidateGeneratedClass(Class);
+	bool Result = UScriptBlueprint::ValidateGeneratedClass(Class);
+	return SuperResult && Result;
+}
diff --git a/Engine/Plugins/Script/ScriptEditorPlugin/Source/ScriptEditorPlugin/Private/ScriptBlueprintCompiler.h b/Engine/Plugins/Script/ScriptEditorPlugin/Source/ScriptEditorPlugin/Private/ScriptBlueprintCompiler.h
new file mode 100644
--- /dev/null
+++ b/Engine/Plugins/Script/ScriptEditorPlugin/Source/ScriptEditorPlugin/Private/ScriptBlueprintCompiler.h
@@ -0,0 +1,44 @@
+// Copyright 1998-2014 Epic Games, Inc. All Rights Reserved.
+#pragma once
+#include "KismetCompiler.h"
+
+/**
+* Script Blueprint Compiler
+*/
+class FScriptBlueprintCompiler : public FKismetCompilerContext
+{
+protected:
+
+	typedef FKismetCompilerContext Super;
+
+public:
+	FScriptBlueprintCompiler(UScriptBlueprint* SourceSketch, FCompilerResultsLog& InMessageLog, const FKismetCompilerOptions& InCompilerOptions, TArray<UObject*>* InObjLoaded);
+	virtual ~FScriptBlueprintCompiler();
+
+	// FKismetCompilerContext
+	virtual void Compile() override;
+	// End FKismetCompilerContext
+
+protected:
+	UScriptBlueprint* ScriptBlueprint() const { return Cast<UScriptBlueprint>(Blueprint); }
+
+	// FKismetCompilerContext
+	virtual void SpawnNewClass(const FString& NewClassName) override;
+	virtual void CleanAndSanitizeClass(UBlueprintGeneratedClass* ClassToClean, UObject*& OldCDO) override;
+	virtual void EnsureProperGeneratedClass(UClass*& TargetClass) override;
+	virtual void CreateClassVariablesFromBlueprint() override;
+	virtual void FinishCompilingClass(UClass* Class) override;
+	virtual bool ValidateGeneratedClass(UBlueprintGeneratedClass* Class) override;
+	virtual void CopyTermDefaultsToDefaultObject(UObject* DefaultObject) override;
+	// End FKismetCompilerContext
+
+protected:
+
+	/** New script class */
+	UScriptBlueprintGeneratedClass* NewScriptBlueprintClass;
+	/** Script context */
+	TAutoPtr<FScriptContextBase> ScriptContext;
+	/** Script-defined properties and functions */
+	TArray<FScriptField> ScriptDefinedFields;
+};
+
diff --git a/Engine/Plugins/Script/ScriptEditorPlugin/Source/ScriptEditorPlugin/Private/ScriptBlueprintEditor.cpp b/Engine/Plugins/Script/ScriptEditorPlugin/Source/ScriptEditorPlugin/Private/ScriptBlueprintEditor.cpp
new file mode 100644
--- /dev/null
+++ b/Engine/Plugins/Script/ScriptEditorPlugin/Source/ScriptEditorPlugin/Private/ScriptBlueprintEditor.cpp
@@ -0,0 +1,50 @@
+// Copyright 1998-2014 Epic Games, Inc. All Rights Reserved.
+
+#include "ScriptEditorPluginPrivatePCH.h"
+#include "SKismetInspector.h"
+#include "ScriptBlueprintEditor.h"
+#include "BlueprintEditorModes.h"
+
+#define LOCTEXT_NAMESPACE "ScriptEditor"
+
+FScriptBlueprintEditor::FScriptBlueprintEditor()
+{
+}
+
+FScriptBlueprintEditor::~FScriptBlueprintEditor()
+{
+	UScriptBlueprint* Blueprint = GetScriptBlueprintObj();
+	if ( Blueprint )
+	{
+		Blueprint->OnChanged().RemoveAll(this);
+	}
+}
+
+void FScriptBlueprintEditor::InitScriptBlueprintEditor(const EToolkitMode::Type Mode, const TSharedPtr< IToolkitHost >& InitToolkitHost, const TArray<UBlueprint*>& InBlueprints, bool bShouldOpenInDefaultsMode)
+{
+	InitBlueprintEditor(Mode, InitToolkitHost, InBlueprints, bShouldOpenInDefaultsMode);
+	SetCurrentMode(FBlueprintEditorApplicationModes::BlueprintDefaultsMode);
+}
+
+void FScriptBlueprintEditor::AddApplicationMode(FName ModeName, TSharedRef<FApplicationMode> Mode)
+{
+	if (ModeName == FBlueprintEditorApplicationModes::BlueprintDefaultsMode)
+	{
+		FBlueprintEditor::AddApplicationMode(ModeName, Mode);
+	}
+}
+
+UScriptBlueprint* FScriptBlueprintEditor::GetScriptBlueprintObj() const
+{
+	return Cast<UScriptBlueprint>(GetBlueprintObj());
+}
+
+UBlueprint* FScriptBlueprintEditor::GetBlueprintObj() const
+{
+	auto Blueprint = FBlueprintEditor::GetBlueprintObj();
+	//auto ScriptBlueprint = Cast<UScriptBlueprint>(Blueprint);
+	//ScriptBlueprint->
+	return Blueprint;
+}
+
+#undef LOCTEXT_NAMESPACE
diff --git a/Engine/Plugins/Script/ScriptEditorPlugin/Source/ScriptEditorPlugin/Private/ScriptBlueprintEditor.h b/Engine/Plugins/Script/ScriptEditorPlugin/Source/ScriptEditorPlugin/Private/ScriptBlueprintEditor.h
new file mode 100644
--- /dev/null
+++ b/Engine/Plugins/Script/ScriptEditorPlugin/Source/ScriptEditorPlugin/Private/ScriptBlueprintEditor.h
@@ -0,0 +1,26 @@
+// Copyright 1998-2014 Epic Games, Inc. All Rights Reserved.
+
+#pragma once
+
+#include "Editor/Kismet/Public/BlueprintEditor.h"
+
+/**
+ * Script blueprint editor (extends Blueprint editor)
+ */
+class FScriptBlueprintEditor : public FBlueprintEditor
+{
+public:
+	FScriptBlueprintEditor();
+	virtual ~FScriptBlueprintEditor();
+
+	/** Initializes script editor */
+	void InitScriptBlueprintEditor(const EToolkitMode::Type Mode, const TSharedPtr< IToolkitHost >& InitToolkitHost, const TArray<UBlueprint*>& InBlueprints, bool bShouldOpenInDefaultsMode);
+
+	/** Work around for the fact there's no other way to prevent Components Mode and Graph Mode from being entered/created */
+	virtual void AddApplicationMode(FName ModeName, TSharedRef<FApplicationMode> Mode) override;
+
+	/** @return The Script blueprint currently being edited in this editor */
+	class UScriptBlueprint* GetScriptBlueprintObj() const;
+
+	virtual UBlueprint* GetBlueprintObj() const override;
+};
diff --git a/Engine/Plugins/Script/ScriptEditorPlugin/Source/ScriptEditorPlugin/Private/ScriptEditorPlugin.cpp b/Engine/Plugins/Script/ScriptEditorPlugin/Source/ScriptEditorPlugin/Private/ScriptEditorPlugin.cpp
--- a/Engine/Plugins/Script/ScriptEditorPlugin/Source/ScriptEditorPlugin/Private/ScriptEditorPlugin.cpp
+++ b/Engine/Plugins/Script/ScriptEditorPlugin/Source/ScriptEditorPlugin/Private/ScriptEditorPlugin.cpp
@@ -1,23 +1,89 @@
 // Copyright 1998-2014 Epic Games, Inc. All Rights Reserved.
 
 #include "ScriptEditorPluginPrivatePCH.h"
+#include "ScriptBlueprintCompiler.h"
+#include "KismetCompilerModule.h"
 
 DEFINE_LOG_CATEGORY(LogScriptEditorPlugin);
 
+/**
+* Script blueprint editor actions
+*/
+class FAssetTypeActions_ScriptClass : public FAssetTypeActions_Base
+{
+public:
+	// IAssetTypeActions Implementation
+	virtual FText GetName() const override { return NSLOCTEXT("AssetTypeActions", "AssetTypeActions_ScriptClass", "Script Class"); }
+	virtual FColor GetTypeColor() const override { return FColor(201, 29, 85); }
+	virtual UClass* GetSupportedClass() const override { return UScriptBlueprint::StaticClass(); }
+	virtual uint32 GetCategories() override { return EAssetTypeCategories::Misc; }
+	virtual void OpenAssetEditor(const TArray<UObject*>& InObjects, TSharedPtr<class IToolkitHost> EditWithinLevelEditor = TSharedPtr<IToolkitHost>()) override
+	{
+		EToolkitMode::Type Mode = EditWithinLevelEditor.IsValid() ? EToolkitMode::WorldCentric : EToolkitMode::Standalone;
+
+		for (auto ObjIt = InObjects.CreateConstIterator(); ObjIt; ++ObjIt)
+		{
+			auto Blueprint = Cast<UBlueprint>(*ObjIt);
+			if (Blueprint && Blueprint->SkeletonGeneratedClass && Blueprint->GeneratedClass)
+			{
+				TSharedRef<FScriptBlueprintEditor> NewBlueprintEditor(new FScriptBlueprintEditor());
+
+				TArray<UBlueprint*> Blueprints;
+				Blueprints.Add(Blueprint);
+				NewBlueprintEditor->InitScriptBlueprintEditor(Mode, EditWithinLevelEditor, Blueprints, true);
+			}
+			else
+			{
+				FMessageDialog::Open(EAppMsgType::Ok, 
+					FText::FromString(TEXT("Script Blueprint could not be loaded because it derives from an invalid class. Check to make sure the parent class for this blueprint hasn't been removed!")));
+			}
+		}
+	}
+};
+
+/**
+* Script blueprint module
+*/
 class FScriptEditorPlugin : public IScriptEditorPlugin
 {
 	/** IModuleInterface implementation */
 	virtual void StartupModule() override;
 	virtual void ShutdownModule() override;
+
+private:
+
+	/** Callback for the script compiler */
+	FReply CompileScript(UBlueprint* Blueprint, const FKismetCompilerOptions& CompileOptions, FCompilerResultsLog& Results, TArray<UObject*>* ObjLoaded);
 };
 
 IMPLEMENT_MODULE(FScriptEditorPlugin, ScriptEditorPlugin)
 
-
 void FScriptEditorPlugin::StartupModule()
 {
+	// Register asset types
+	FModuleManager::LoadModuleChecked<FAssetToolsModule>("AssetTools").Get().RegisterAssetTypeActions(MakeShareable(new FAssetTypeActions_ScriptClass));
+
+	// Register blueprint compiler
+	IKismetCompilerInterface& KismetCompilerModule = FModuleManager::LoadModuleChecked<IKismetCompilerInterface>("KismetCompiler");
+	FBlueprintCompileDelegate CompileDelegate;
+	CompileDelegate.BindRaw(this, &FScriptEditorPlugin::CompileScript);
+	KismetCompilerModule.GetCompilers().Add(CompileDelegate);
 }
 
 void FScriptEditorPlugin::ShutdownModule()
 {
 }
+
+FReply FScriptEditorPlugin::CompileScript(UBlueprint* Blueprint, const FKismetCompilerOptions& CompileOptions, FCompilerResultsLog& Results, TArray<UObject*>* ObjLoaded)
+{
+	if (UScriptBlueprint* ScriptBlueprint = Cast<UScriptBlueprint>(Blueprint))
+	{
+		FScriptBlueprintCompiler Compiler(ScriptBlueprint, Results, CompileOptions, ObjLoaded);
+		Compiler.Compile();
+		check(Compiler.NewClass);
+
+		return FReply::Handled();
+	}
+
+	return FReply::Unhandled();
+}
\ No newline at end of file
diff --git a/Engine/Plugins/Script/ScriptEditorPlugin/Source/ScriptEditorPlugin/Private/ScriptEditorPluginPrivatePCH.h b/Engine/Plugins/Script/ScriptEditorPlugin/Source/ScriptEditorPlugin/Private/ScriptEditorPluginPrivatePCH.h
--- a/Engine/Plugins/Script/ScriptEditorPlugin/Source/ScriptEditorPlugin/Private/ScriptEditorPluginPrivatePCH.h
+++ b/Engine/Plugins/Script/ScriptEditorPlugin/Source/ScriptEditorPlugin/Private/ScriptEditorPluginPrivatePCH.h
@@ -6,9 +6,16 @@
 #include "ModuleManager.h"
 #include "Engine.h"
 #include "UnrealEd.h"
-#include "ScriptAsset.h"
+#include "ScriptBlueprintGeneratedClass.h"
 #include "ScriptFactory.h"
 #include "ReimportScriptFactory.h"
+#include "AssetTypeActions_Base.h"
+#include "ScriptActor.h"
+#include "ScriptComponent.h"
+#include "ClassViewerModule.h"
+#include "ClassViewerFilter.h"
+#include "Kismet2/KismetEditorUtilities.h"
+#include "Kismet2/SClassPickerDialog.h"
 
 DECLARE_LOG_CATEGORY_EXTERN(LogScriptEditorPlugin, Log, All);
 
diff --git a/Engine/Plugins/Script/ScriptEditorPlugin/Source/ScriptEditorPlugin/Private/ScriptFactory.cpp b/Engine/Plugins/Script/ScriptEditorPlugin/Source/ScriptEditorPlugin/Private/ScriptFactory.cpp
--- a/Engine/Plugins/Script/ScriptEditorPlugin/Source/ScriptEditorPlugin/Private/ScriptFactory.cpp
+++ b/Engine/Plugins/Script/ScriptEditorPlugin/Source/ScriptEditorPlugin/Private/ScriptFactory.cpp
@@ -1,12 +1,33 @@
 // Copyright 1998-2014 Epic Games, Inc. All Rights Reserved. 
 #include "ScriptEditorPluginPrivatePCH.h"
 
+/**
+ * Class Picker filter
+ */
+class FScriptFactoryFilter : public IClassViewerFilter
+{
+public:
+	TSet< const UClass* > AllowedChildOfClasses;
+
+	bool IsClassAllowed(const FClassViewerInitializationOptions& InInitOptions, const UClass* InClass, TSharedRef< FClassViewerFilterFuncs > InFilterFuncs)
+	{
+		return AllowedChildOfClasses.Contains(InClass);
+	}
+
+	virtual bool IsUnloadedClassAllowed(const FClassViewerInitializationOptions& InInitOptions, const TSharedRef< const IUnloadedBlueprintData > InUnloadedClassData, TSharedRef< FClassViewerFilterFuncs > InFilterFuncs) override
+	{
+		return true;
+	}
+};
+
 //////////////////////////////////////////////////////////////////////////
 
+
 UScriptFactory::UScriptFactory(const FPostConstructInitializeProperties& PCIP)
 	: Super( PCIP )
 {
-	SupportedClass = UScriptAsset::StaticClass();
+	SupportedClass = UScriptBlueprint::StaticClass();
+	ParentClass = AScriptActor::StaticClass();
 
 	Formats.Add(TEXT("txt;Script"));
 	Formats.Add(TEXT("lua;Script"));
@@ -14,6 +35,7 @@ UScriptFactory::UScriptFactory(const FPostConstructInitializeProperties& PCIP)
 	bCreateNew = false;
 	bEditorImport = true;
 	bText = true;
+	bEditAfterNew = true;	
 }
 
 void UScriptFactory::PostInitProperties()
@@ -21,23 +43,60 @@ void UScriptFactory::PostInitProperties()
 	Super::PostInitProperties();
 }
 
+bool UScriptFactory::ConfigureProperties()
+{
+	// Null the parent class so we can check for selection later
+	ParentClass = NULL;
+
+	// Load the class viewer module to display a class picker
+	FClassViewerModule& ClassViewerModule = FModuleManager::LoadModuleChecked<FClassViewerModule>("ClassViewer");
+
+	// Fill in options
+	FClassViewerInitializationOptions Options;
+	Options.Mode = EClassViewerMode::ClassPicker;
+	Options.DisplayMode = EClassViewerDisplayMode::ListView;
+	Options.bShowObjectRootClass = false;
+	Options.bIsBlueprintBaseOnly = false;
+	Options.bShowUnloadedBlueprints = false;
+
+	// Only allow script-enabled classes
+	TSharedPtr<FScriptFactoryFilter> Filter = MakeShareable(new FScriptFactoryFilter);
+	Options.ClassFilter = Filter;
+	Filter->AllowedChildOfClasses.Add(AScriptActor::StaticClass());
+	Filter->AllowedChildOfClasses.Add(UScriptComponent::StaticClass());
+
+	const FText TitleText = NSLOCTEXT("EditorFactories", "CreateScriptOptions", "Pick Parent Class");
+	UClass* ChosenClass = NULL;
+	const bool bPressedOk = SClassPickerDialog::PickClass(TitleText, Options, ChosenClass, UScriptBlueprintGeneratedClass::StaticClass());
+	if (bPressedOk)
+	{
+		ParentClass = ChosenClass;
+	}
+
+	return bPressedOk;
+}
+
+
 bool UScriptFactory::DoesSupportClass(UClass* Class)
 {
-	return Class == UScriptAsset::StaticClass();
+	return Class == UScriptBlueprint::StaticClass();
 }
 
 UObject* UScriptFactory::FactoryCreateText(UClass* InClass, UObject* InParent, FName InName, EObjectFlags Flags, UObject* Context, const TCHAR* Type, const TCHAR*& Buffer, const TCHAR* BufferEnd, FFeedbackContext* Warn)
 {
-	UScriptAsset* ScriptAsset = NewNamedObject<UScriptAsset>(InParent, InName, Flags);
-	
-	ScriptAsset->SourceFilePath = FReimportManager::SanitizeImportFilename(CurrentFilename, ScriptAsset);
-	ScriptAsset->SourceFileTimestamp = IFileManager::Get().GetTimeStamp(*CurrentFilename).ToString();
-	ScriptAsset->SourceCode = Buffer;
+	GEditor->SelectNone(true, true, false);
+
+	UScriptBlueprint* NewBlueprint = CastChecked<UScriptBlueprint>(FKismetEditorUtilities::CreateBlueprint(ParentClass, InParent, InName, BPTYPE_Normal, UScriptBlueprint::StaticClass(), UScriptBlueprintGeneratedClass::StaticClass(), "UScriptFactory"));
+	NewBlueprint->SourceFilePath = FReimportManager::SanitizeImportFilename(CurrentFilename, NewBlueprint);
+	NewBlueprint->SourceCode = Buffer;
+	NewBlueprint->SourceFileTimestamp = IFileManager::Get().GetTimeStamp(*NewBlueprint->SourceFilePath).ToString();
+
+	// Need to make sure we compile with the new source code
+	FKismetEditorUtilities::CompileBlueprint(NewBlueprint);
 
-	FEditorDelegates::OnAssetPostImport.Broadcast(this, ScriptAsset);
-	ScriptAsset->PostEditChange();
+	FEditorDelegates::OnAssetPostImport.Broadcast(this, NewBlueprint);
 
-	return ScriptAsset;
+	return NewBlueprint;
 }
 
 
@@ -48,13 +107,17 @@ UReimportScriptFactory::UReimportScriptFactory(const class FPostConstructInitial
 {
 }
 
+bool UReimportScriptFactory::ConfigureProperties()
+{
+	return UFactory::ConfigureProperties();
+}
 
 bool UReimportScriptFactory::CanReimport(UObject* Obj, TArray<FString>& OutFilenames)
 {
-	UScriptAsset* ScriptAsset = Cast<UScriptAsset>(Obj);
-	if (ScriptAsset)
+	UScriptBlueprint* ScriptClass = Cast<UScriptBlueprint>(Obj);
+	if (ScriptClass)
 	{
-		OutFilenames.Add(FReimportManager::ResolveImportFilename(ScriptAsset->SourceFilePath, ScriptAsset));
+		OutFilenames.Add(FReimportManager::ResolveImportFilename(ScriptClass->SourceFilePath, ScriptClass));
 		return true;
 	}
 	return false;
@@ -62,10 +125,10 @@ bool UReimportScriptFactory::CanReimport(UObject* Obj, TArray<FString>& OutFilen
 
 void UReimportScriptFactory::SetReimportPaths(UObject* Obj, const TArray<FString>& NewReimportPaths)
 {
-	UScriptAsset* ScriptAsset = Cast<UScriptAsset>(Obj);
-	if (ScriptAsset && ensure(NewReimportPaths.Num() == 1))
+	UScriptBlueprint* ScriptClass = Cast<UScriptBlueprint>(Obj);
+	if (ScriptClass && ensure(NewReimportPaths.Num() == 1))
 	{
-		ScriptAsset->SourceFilePath = FReimportManager::SanitizeImportFilename(NewReimportPaths[0], Obj);
+		ScriptClass->SourceFilePath = FReimportManager::SanitizeImportFilename(NewReimportPaths[0], Obj);
 	}
 }
 
@@ -74,15 +137,15 @@ void UReimportScriptFactory::SetReimportPaths(UObject* Obj, const TArray<FString
 */
 EReimportResult::Type UReimportScriptFactory::Reimport(UObject* Obj)
 {
-	UScriptAsset* ScriptAsset = Cast<UScriptAsset>(Obj);
-	if (!ScriptAsset)
+	UScriptBlueprint* ScriptClass = Cast<UScriptBlueprint>(Obj);
+	if (!ScriptClass)
 	{
 		return EReimportResult::Failed;
 	}
 
-	TGuardValue<UScriptAsset*> OriginalScriptGuardValue(OriginalScript, ScriptAsset);
+	TGuardValue<UScriptBlueprint*> OriginalScriptGuardValue(OriginalScript, ScriptClass);
 
-	const FString ResolvedSourceFilePath = FReimportManager::ResolveImportFilename(ScriptAsset->SourceFilePath, ScriptAsset);
+	const FString ResolvedSourceFilePath = FReimportManager::ResolveImportFilename(ScriptClass->SourceFilePath, ScriptClass);
 	if (!ResolvedSourceFilePath.Len())
 	{
 		return EReimportResult::Failed;
@@ -97,17 +160,17 @@ EReimportResult::Type UReimportScriptFactory::Reimport(UObject* Obj)
 		return EReimportResult::Failed;
 	}
 
-	if (UFactory::StaticImportObject(ScriptAsset->GetClass(), ScriptAsset->GetOuter(), *ScriptAsset->GetName(), RF_Public | RF_Standalone, *ResolvedSourceFilePath, NULL, this))
+	if (UFactory::StaticImportObject(ScriptClass->GetClass(), ScriptClass->GetOuter(), *ScriptClass->GetName(), RF_Public | RF_Standalone, *ResolvedSourceFilePath, NULL, this))
 	{
 		UE_LOG(LogScriptEditorPlugin, Log, TEXT("Imported successfully"));
 		// Try to find the outer package so we can dirty it up
-		if (ScriptAsset->GetOuter())
+		if (ScriptClass->GetOuter())
 		{
-			ScriptAsset->GetOuter()->MarkPackageDirty();
+			ScriptClass->GetOuter()->MarkPackageDirty();
 		}
 		else
 		{
-			ScriptAsset->MarkPackageDirty();
+			ScriptClass->MarkPackageDirty();
 		}
 	}
 	else
diff --git a/Engine/Plugins/Script/ScriptEditorPlugin/Source/ScriptEditorPlugin/ScriptEditorPlugin.Build.cs b/Engine/Plugins/Script/ScriptEditorPlugin/Source/ScriptEditorPlugin/ScriptEditorPlugin.Build.cs
--- a/Engine/Plugins/Script/ScriptEditorPlugin/Source/ScriptEditorPlugin/ScriptEditorPlugin.Build.cs
+++ b/Engine/Plugins/Script/ScriptEditorPlugin/Source/ScriptEditorPlugin/ScriptEditorPlugin.Build.cs
@@ -27,7 +27,11 @@ namespace UnrealBuildTool.Rules
 					"Engine",
 					"InputCore",
 					"UnrealEd",
-					"ScriptPlugin"
+					"AssetTools",
+					"ScriptPlugin",
+					"ClassViewer",
+					"KismetCompiler",
+					"Kismet"
 					// ... add other public dependencies that you statically link with here ...
 				}
 				);
diff --git a/Engine/Plugins/Script/ScriptGeneratorPlugin/Source/ScriptGeneratorPlugin/Private/LuaScriptCodeGenerator.cpp b/Engine/Plugins/Script/ScriptGeneratorPlugin/Source/ScriptGeneratorPlugin/Private/LuaScriptCodeGenerator.cpp
--- a/Engine/Plugins/Script/ScriptGeneratorPlugin/Source/ScriptGeneratorPlugin/Private/LuaScriptCodeGenerator.cpp
+++ b/Engine/Plugins/Script/ScriptGeneratorPlugin/Source/ScriptGeneratorPlugin/Private/LuaScriptCodeGenerator.cpp
@@ -3,15 +3,17 @@
 #include "LuaScriptCodeGenerator.h"
 
 // Supported structs
+static FName Name_Vector2D("Vector2D");
 static FName Name_Vector("Vector");
 static FName Name_Vector4("Vector4");
 static FName Name_Quat("Quat");
 static FName Name_Transform("Transform");
+static FName Name_LinearColor("LinearColor");
+static FName Name_Color("Color");
 
 FLuaScriptCodeGenerator::FLuaScriptCodeGenerator(const FString& RootLocalPath, const FString& RootBuildPath, const FString& OutputDirectory, const FString& InIncludeBase)
 : FScriptCodeGeneratorBase(RootLocalPath, RootBuildPath, OutputDirectory, InIncludeBase)
 {
-
 }
 
 FString FLuaScriptCodeGenerator::GenerateWrapperFunctionDeclaration(const FString& ClassNameCPP, UClass* Class, UFunction* Function)
@@ -55,7 +57,11 @@ FString FLuaScriptCodeGenerator::InitializeFunctionDispatchParam(UFunction* Func
 		else if (Param->IsA(UStructProperty::StaticClass()))
 		{
 			UStructProperty* StructProp = CastChecked<UStructProperty>(Param);
-			if (StructProp->Struct->GetFName() == Name_Vector)
+			if (StructProp->Struct->GetFName() == Name_Vector2D)
+			{
+				Initializer = TEXT("(FLuaVector2D::Get");
+			}
+			else if (StructProp->Struct->GetFName() == Name_Vector)
 			{
 				Initializer = TEXT("(FLuaVector::Get");
 			}
@@ -67,6 +73,14 @@ FString FLuaScriptCodeGenerator::InitializeFunctionDispatchParam(UFunction* Func
 			{
 				Initializer = TEXT("(FLuaQuat::Get");
 			}
+			else if (StructProp->Struct->GetFName() == Name_LinearColor)
+			{
+				Initializer = TEXT("(FLuaLinearColor::Get");
+			}
+			else if (StructProp->Struct->GetFName() == Name_Color)
+			{
+				Initializer = TEXT("FColor(FLuaLinearColor::Get");
+			}
 			else if (StructProp->Struct->GetFName() == Name_Transform)
 			{
 				Initializer = TEXT("(FLuaTransform::Get");
@@ -130,7 +144,11 @@ FString FLuaScriptCodeGenerator::GenerateReturnValueHandler(const FString& Class
 		else if (ReturnValue->IsA(UStructProperty::StaticClass()))
 		{
 			UStructProperty* StructProp = CastChecked<UStructProperty>(ReturnValue);
-			if (StructProp->Struct->GetFName() == Name_Vector)
+			if (StructProp->Struct->GetFName() == Name_Vector2D)
+			{
+				Initializer = FString::Printf(TEXT("FLuaVector2D::Return(InScriptContext, %s);"), *ReturnValueName);
+			}
+			else if (StructProp->Struct->GetFName() == Name_Vector)
 			{
 				Initializer = FString::Printf(TEXT("FLuaVector::Return(InScriptContext, %s);"), *ReturnValueName);
 			}
@@ -142,6 +160,14 @@ FString FLuaScriptCodeGenerator::GenerateReturnValueHandler(const FString& Class
 			{
 				Initializer = FString::Printf(TEXT("FLuaQuat::Return(InScriptContext, %s);"), *ReturnValueName);
 			}
+			else if (StructProp->Struct->GetFName() == Name_LinearColor)
+			{
+				Initializer = FString::Printf(TEXT("FLuaLinearColor::Return(InScriptContext, %s);"), *ReturnValueName);
+			}
+			else if (StructProp->Struct->GetFName() == Name_Color)
+			{
+				Initializer = FString::Printf(TEXT("FLuaLinearColor::Return(InScriptContext, FLinearColor(%s));"), *ReturnValueName);
+			}
 			else if (StructProp->Struct->GetFName() == Name_Transform)
 			{
 				Initializer = FString::Printf(TEXT("FLuaTransform::Return(InScriptContext, %s);"), *ReturnValueName);
@@ -278,9 +304,12 @@ bool FLuaScriptCodeGenerator::IsPropertyTypeSupported(UProperty* Property) const
 	if (Property->IsA(UStructProperty::StaticClass()))
 	{
 		UStructProperty* StructProp = CastChecked<UStructProperty>(Property);
-		if (StructProp->Struct->GetFName() != Name_Vector &&
+		if (StructProp->Struct->GetFName() != Name_Vector2D &&
+			StructProp->Struct->GetFName() != Name_Vector &&
 			StructProp->Struct->GetFName() != Name_Vector4 &&
 			StructProp->Struct->GetFName() != Name_Quat &&
+			StructProp->Struct->GetFName() != Name_LinearColor &&
+			StructProp->Struct->GetFName() != Name_Color &&
 			StructProp->Struct->GetFName() != Name_Transform)
 		{
 			bSupported = false;
diff --git a/Engine/Plugins/Script/ScriptPlugin/Source/ScriptPlugin/Classes/ScriptActor.h b/Engine/Plugins/Script/ScriptPlugin/Source/ScriptPlugin/Classes/ScriptActor.h
new file mode 100644
--- /dev/null
+++ b/Engine/Plugins/Script/ScriptPlugin/Source/ScriptPlugin/Classes/ScriptActor.h
@@ -0,0 +1,35 @@
+// Copyright 1998-2014 Epic Games, Inc. All Rights Reserved.
+#pragma once
+
+#include "ScriptBlueprintGeneratedClass.h"
+#include "ScriptActor.generated.h"
+
+/**
+* Script-extendable actor class
+*/
+UCLASS(BlueprintType, Abstract, EarlyAccessPreview)
+class SCRIPTPLUGIN_API AScriptActor : public AActor
+{
+	GENERATED_UCLASS_BODY()
+
+	// UObject interface
+	virtual void PostLoad() override;
+	virtual void BeginDestroy() override;
+
+	// AActor interface
+	virtual void OnConstruction(const FTransform& Transform) override;
+	virtual void BeginPlay() override;
+	virtual void TickActor(float DeltaTime, enum ELevelTick TickType, FActorTickFunction& ThisTickFunction) override;
+	virtual void ReceiveDestroyed() override;
+	virtual void ReceiveEndPlay(EEndPlayReason::Type EndPlayReason) override;
+
+	/** 
+	 * Invokes script-defined function
+	 */
+	void InvokeScriptFunction(FFrame& Stack, RESULT_DECL);
+
+private:
+
+	/** Script context (code) for this actor */
+	FScriptContextBase* ScriptContext;
+};
diff --git a/Engine/Plugins/Script/ScriptPlugin/Source/ScriptPlugin/Classes/ScriptAsset.h b/Engine/Plugins/Script/ScriptPlugin/Source/ScriptPlugin/Classes/ScriptAsset.h
deleted file mode 100644
diff --git a/Engine/Plugins/Script/ScriptPlugin/Source/ScriptPlugin/Classes/ScriptBlueprint.h b/Engine/Plugins/Script/ScriptPlugin/Source/ScriptPlugin/Classes/ScriptBlueprint.h
new file mode 100644
--- /dev/null
+++ b/Engine/Plugins/Script/ScriptPlugin/Source/ScriptPlugin/Classes/ScriptBlueprint.h
@@ -0,0 +1,57 @@
+// Copyright 1998-2014 Epic Games, Inc. All Rights Reserved.
+
+#pragma once
+
+#include "Engine/Blueprint.h"
+#include "ScriptBlueprintGeneratedClass.h"
+#include "ScriptBlueprint.generated.h"
+
+/**
+ * The Script blueprint enables extending AScriptActor and UScriptComponent classes
+ */
+UCLASS(BlueprintType, EarlyAccessPreview)
+class SCRIPTPLUGIN_API UScriptBlueprint : public UBlueprint
+{
+	GENERATED_UCLASS_BODY()
+
+public:
+
+	/** Path to the script text used to construct this asset. Relative to the object's package, BaseDir() or absolute */
+	UPROPERTY(Category = Script, VisibleAnywhere, BlueprintReadWrite)
+	FString SourceFilePath;
+
+	/** Date/Time-stamp of the file from the last import. */
+	UPROPERTY(Category = Script, VisibleAnywhere)
+	FString SourceFileTimestamp;
+
+	/** Generated script bytecode */
+	UPROPERTY()
+	TArray<uint8> ByteCode;
+
+	/** Script source code. @todo: this should be editor-only */
+	UPROPERTY()
+	FString SourceCode;
+	 
+	virtual void PostLoad() override;
+
+	// UBlueprint interface
+	virtual UClass* RegenerateClass(UClass* ClassToRegenerate, UObject* PreviousCDO, TArray<UObject*>& ObjLoaded) override;
+
+#if WITH_EDITOR
+	virtual UClass* GetBlueprintClass() const override;
+
+	virtual bool SupportedByDefaultBlueprintFactory() const override
+	{
+		return false;
+	}
+	// End of UBlueprint interface
+
+	static bool ValidateGeneratedClass(const UClass* InClass);
+	void UpdateScriptStatus();
+
+	/**
+	 * Checks if the source code changed and updates it if so (does not recompile the class)
+	 */
+	void UpdateSourceCodeIfChanged();
+#endif
+};
\ No newline at end of file
diff --git a/Engine/Plugins/Script/ScriptPlugin/Source/ScriptPlugin/Classes/ScriptBlueprintGeneratedClass.h b/Engine/Plugins/Script/ScriptPlugin/Source/ScriptPlugin/Classes/ScriptBlueprintGeneratedClass.h
new file mode 100644
--- /dev/null
+++ b/Engine/Plugins/Script/ScriptPlugin/Source/ScriptPlugin/Classes/ScriptBlueprintGeneratedClass.h
@@ -0,0 +1,122 @@
+// Copyright 1998-2014 Epic Games, Inc. All Rights Reserved.
+
+#pragma once
+
+#include "Engine/BlueprintGeneratedClass.h"
+#include "ScriptBlueprintGeneratedClass.generated.h"
+
+/**
+* Script-defined field (variable or function)
+*/
+struct FScriptField
+{
+	/** Field name */
+	FName Name;
+	/** Field type */
+	UClass* Class;
+
+	FScriptField()
+		: Class(NULL)
+	{
+	}
+	FScriptField(FName InName, UClass* InClass)
+		: Name(InName)
+		, Class(InClass)
+	{
+	}
+};
+
+/** Script context for this component */
+class FScriptContextBase
+{
+public:
+	/**
+	 * Initializes script context given script code
+	 * @param Code Script code
+	 * @param Owner Owner of this context
+	 */
+	virtual bool Initialize(const FString& Code, UObject* Owner) = 0;
+
+	/**
+	* Sends BeginPlay event to the script
+	*/
+	virtual void BeginPlay() = 0;
+
+	/**
+	* Sends Tick event to the script
+	*/
+	virtual void Tick(float DeltaTime) = 0;
+
+	/**
+	* Sends Destroy event to the script and destroys the script.
+	*/
+	virtual void Destroy() = 0;
+
+	/**
+	* @return true if script defines Tick function
+	*/
+	virtual bool CanTick() = 0;
+
+	/**
+	* Calls arbitrary script function (no arguments) given its name.
+	* @param FunctionName Name of the function to call
+	*/
+	virtual bool CallFunction(const FString&  FunctionName) = 0;
+
+	virtual bool SetFloatVariable(const FString&  GlobalVariable, float NewValue) = 0;
+	virtual float GetFloatVariable(const FString&  GlobalVariable) = 0;
+	virtual bool SetStringVariable(const FString&  GlobalVariable, const FString& NewValue) = 0;
+	virtual FString GetStringVariable(const FString&  GlobalVariable) = 0;
+
+	/**
+	* Invokes script function from Blueprint code
+	*/
+	virtual void InvokeScriptFunction(FFrame& Stack, RESULT_DECL) = 0;
+
+#if WITH_EDITOR
+	/**
+	* Returns a list of exported fields from the script (member variables and functions).
+	*/
+	virtual void GetScriptDefinedFields(TArray<FScriptField>& OutFields) = 0;
+#endif
+
+	virtual ~FScriptContextBase() {}
+};
+
+/**
+* Script generated class
+*/
+UCLASS(EarlyAccessPreview)
+class SCRIPTPLUGIN_API UScriptBlueprintGeneratedClass : public UBlueprintGeneratedClass
+{
+	GENERATED_UCLASS_BODY()
+
+	/** Generated script bytecode */
+	UPROPERTY()
+	TArray<uint8> ByteCode;
+
+	/** Script source code. @todo: this should be editor-only */
+	UPROPERTY()
+	FString SourceCode;
+
+	virtual void PostInitProperties() override;
+	virtual void Link(FArchive& Ar, bool bRelinkExistingProperties) override;
+
+	/**
+	 * Creates a script context object
+	 */
+	FScriptContextBase* CreateContext();
+
+	/**
+	* Adds a unique native function mapping
+	* @param InName Name of the native function
+	* @param InPointer Pointer to the native member function
+	*/
+	void AddUniqueNativeFunction(const FName& InName, Native InPointer);
+
+	/**
+	* Removes native function mapping
+	* @param InName Name of the native function
+	*/
+	void RemoveNativeFunction(const FName& InName);
+};
\ No newline at end of file
diff --git a/Engine/Plugins/Script/ScriptPlugin/Source/ScriptPlugin/Classes/ScriptComponent.h b/Engine/Plugins/Script/ScriptPlugin/Source/ScriptPlugin/Classes/ScriptComponent.h
--- a/Engine/Plugins/Script/ScriptPlugin/Source/ScriptPlugin/Classes/ScriptComponent.h
+++ b/Engine/Plugins/Script/ScriptPlugin/Source/ScriptPlugin/Classes/ScriptComponent.h
@@ -1,32 +1,25 @@
 // Copyright 1998-2014 Epic Games, Inc. All Rights Reserved.
 
 #pragma once
-#include "ScriptAsset.h"
+#include "ScriptBlueprintGeneratedClass.h"
 #include "ScriptComponent.generated.h"
 
-/** Script context for this component */
-class FScriptContextBase
-{
-public:
-	virtual bool Initialize(class UScriptComponent* Owner) = 0;
-	virtual void BeginPlay() = 0;
-	virtual void Tick(float DeltaTime) = 0;
-	virtual void Destroy() = 0;
-	virtual bool CanTick() = 0;
-	virtual ~FScriptContextBase() {}
-};
-
-/** Component that allows you to specify custom triangle mesh geometry */
-UCLASS(hidecategories=(Object, ActorComponent), editinlinenew, meta=(BlueprintSpawnableComponent), ClassGroup=Script)
+/** 
+ * Script-extendable component class
+ */
+UCLASS(hidecategories = (Object, ActorComponent), editinlinenew, meta = (BlueprintSpawnableComponent), ClassGroup = Script, Abstract, EarlyAccessPreview)
 class SCRIPTPLUGIN_API UScriptComponent : public UActorComponent
 {
 	GENERATED_UCLASS_BODY()
 
 public:
 
-	/** Script code for this component */
-	UPROPERTY(EditAnywhere, BlueprintReadOnly, Category = Script)
-	UScriptAsset* Script;
+	/**
+	* Calls a script-defined function (no arguments)
+	* @param FunctionName Name of the function to call
+	*/
+	UFUNCTION(BlueprintCallable, Category = "Script|Functions")
+	virtual bool CallScriptFunction(FString FunctionName);
 
 	// Begin UActorComponent interface.
 	virtual void OnRegister() override;
@@ -37,7 +30,7 @@ public:
 
 protected:
 
+	/** Script context (code) */
 	FScriptContextBase* Context;
 };
 
-
diff --git a/Engine/Plugins/Script/ScriptPlugin/Source/ScriptPlugin/Classes/ScriptFunction.h b/Engine/Plugins/Script/ScriptPlugin/Source/ScriptPlugin/Classes/ScriptFunction.h
new file mode 100644
--- /dev/null
+++ b/Engine/Plugins/Script/ScriptPlugin/Source/ScriptPlugin/Classes/ScriptFunction.h
@@ -0,0 +1,10 @@
+// Copyright 1998-2014 Epic Games, Inc. All Rights Reserved.
+#pragma once
+#include "ScriptFunction.generated.h"
+
+/** Script-defined function that maps native function to a script call */
+UCLASS(EarlyAccessPreview)
+class SCRIPTPLUGIN_API UScriptFunction : public UFunction
+{
+	GENERATED_UCLASS_BODY()
+};
diff --git a/Engine/Plugins/Script/ScriptPlugin/Source/ScriptPlugin/Classes/ScriptTestActor.h b/Engine/Plugins/Script/ScriptPlugin/Source/ScriptPlugin/Classes/ScriptTestActor.h
--- a/Engine/Plugins/Script/ScriptPlugin/Source/ScriptPlugin/Classes/ScriptTestActor.h
+++ b/Engine/Plugins/Script/ScriptPlugin/Source/ScriptPlugin/Classes/ScriptTestActor.h
@@ -2,8 +2,10 @@
 #pragma once
 #include "ScriptTestActor.generated.h"
 
-/** Definition of a specific ability that is applied to a character. Exists as part of a Trait. */
-UCLASS(BlueprintType, MinimalAPI)
+/**
+* Test Actor. To be removed at some point.
+*/
+UCLASS(BlueprintType, MinimalAPI, EarlyAccessPreview)
 class AScriptTestActor : public AActor
 {
 	GENERATED_UCLASS_BODY()
diff --git a/Engine/Plugins/Script/ScriptPlugin/Source/ScriptPlugin/Private/LuaIntegration.cpp b/Engine/Plugins/Script/ScriptPlugin/Source/ScriptPlugin/Private/LuaIntegration.cpp
--- a/Engine/Plugins/Script/ScriptPlugin/Source/ScriptPlugin/Private/LuaIntegration.cpp
+++ b/Engine/Plugins/Script/ScriptPlugin/Source/ScriptPlugin/Private/LuaIntegration.cpp
@@ -3,6 +3,7 @@
 #include "ScriptPluginPrivatePCH.h"
 #include "LuaIntegration.h"
 #include "ScriptObjectReferencer.h"
+#include "ScriptActor.h"
 
 #if WITH_LUA
 
@@ -132,10 +133,18 @@ struct FLuaUObject
 	{
 		ULevel* Level = NULL;
 		UObject* This = (UObject*)lua_touserdata(LuaState, 1);
-		UActorComponent* ThisComponent = Cast<UActorComponent>(This);
-		if (ThisComponent && ThisComponent->GetOwner())
+		AActor* ThisActor = Cast<AActor>(This);
+		if (!ThisActor)
 		{
-			Level = ThisComponent->GetOwner()->GetLevel();
+			UActorComponent* ThisComponent = Cast<UActorComponent>(This);
+			if (ThisComponent)
+			{
+				ThisActor = ThisComponent->GetOwner();
+			}
+		}
+		if (ThisActor)
+		{
+			Level = ThisActor->GetLevel();
 		}
 		lua_pushlightuserdata(LuaState, Level);
 		return 1;
@@ -246,13 +255,43 @@ luaL_Reg FLuaUObject::UObject_Lib[] =
 	{ NULL, NULL }
 };
 
+FVector2D FLuaVector2D::Get(lua_State* LuaState, int ParamIndex)
+{
+	FVector2D Result;
+
+	luaL_checktype(LuaState, ParamIndex, LUA_TTABLE);
+
+	lua_rawgeti(LuaState, ParamIndex, 1);
+	int Top = lua_gettop(LuaState);
+	Result.X = (float)lua_tonumber(LuaState, Top);
+	lua_pop(LuaState, 1);
+
+	lua_rawgeti(LuaState, ParamIndex, 2);
+	Top = lua_gettop(LuaState);
+	Result.Y = (float)lua_tonumber(LuaState, Top);
+	lua_pop(LuaState, 1);
+
+	return Result;
+}
+
+void FLuaVector2D::Return(lua_State *LuaState, const FVector2D& Value)
+{
+	lua_newtable(LuaState);
+
+	lua_pushnumber(LuaState, Value.X);
+	lua_rawseti(LuaState, -2, 1);
+
+	lua_pushnumber(LuaState, Value.Y);
+	lua_rawseti(LuaState, -2, 2);
+}
+
 FVector FLuaVector::Get(lua_State* LuaState, int ParamIndex)
 {
 	FVector Result;
 
 	luaL_checktype(LuaState, ParamIndex, LUA_TTABLE);
 
-	lua_rawgeti(LuaState, ParamIndex, 1); 
+	lua_rawgeti(LuaState, ParamIndex, 1);
 	int Top = lua_gettop(LuaState);
 	Result.X = (float)lua_tonumber(LuaState, Top);
 	lua_pop(LuaState, 1);
@@ -290,7 +329,7 @@ FVector4 FLuaVector4::Get(lua_State *LuaState, int ParamIndex)
 
 	luaL_checktype(LuaState, ParamIndex, LUA_TTABLE);
 
-	lua_rawgeti(LuaState, ParamIndex, 1); 
+	lua_rawgeti(LuaState, ParamIndex, 1);
 	int Top = lua_gettop(LuaState);
 	Result.X = (float)lua_tonumber(LuaState, Top);
 	lua_pop(LuaState, 1);
@@ -336,7 +375,7 @@ FQuat FLuaQuat::Get(lua_State* LuaState, int ParamIndex)
 
 	luaL_checktype(LuaState, ParamIndex, LUA_TTABLE);
 
-	lua_rawgeti(LuaState, ParamIndex, 1); 
+	lua_rawgeti(LuaState, ParamIndex, 1);
 	int Top = lua_gettop(LuaState);
 	Result.X = (float)lua_tonumber(LuaState, Top);
 	lua_pop(LuaState, 1);
@@ -376,6 +415,53 @@ void FLuaQuat::Return(lua_State* LuaState, const FQuat& Quat)
 	lua_rawseti(LuaState, -2, 4);
 }
 
+FLinearColor FLuaLinearColor::Get(lua_State *LuaState, int ParamIndex)
+{
+	FLinearColor Result;
+
+	luaL_checktype(LuaState, ParamIndex, LUA_TTABLE);
+
+	lua_rawgeti(LuaState, ParamIndex, 1);
+	int Top = lua_gettop(LuaState);
+	Result.R = (float)lua_tonumber(LuaState, Top);
+	lua_pop(LuaState, 1);
+
+	lua_rawgeti(LuaState, ParamIndex, 2);
+	Top = lua_gettop(LuaState);
+	Result.G = (float)lua_tonumber(LuaState, Top);
+	lua_pop(LuaState, 1);
+
+	lua_rawgeti(LuaState, ParamIndex, 3);
+	Top = lua_gettop(LuaState);
+	Result.B = (float)lua_tonumber(LuaState, Top);
+	lua_pop(LuaState, 1);
+
+	lua_rawgeti(LuaState, ParamIndex, 4);
+	Top = lua_gettop(LuaState);
+	Result.A = (float)lua_tonumber(LuaState, Top);
+	lua_pop(LuaState, 1);
+
+	return Result;
+}
+
+void FLuaLinearColor::Return(lua_State* LuaState, const FLinearColor& Value)
+{
+	lua_newtable(LuaState);
+
+	lua_pushnumber(LuaState, Value.R);
+	lua_rawseti(LuaState, -2, 1);
+
+	lua_pushnumber(LuaState, Value.G);
+	lua_rawseti(LuaState, -2, 2);
+
+	lua_pushnumber(LuaState, Value.B);
+	lua_rawseti(LuaState, -2, 3);
+
+	lua_pushnumber(LuaState, Value.A);
+	lua_rawseti(LuaState, -2, 4);
+}
+
+
 template <typename T, typename R>
 static R GetVectorFromTable(lua_State *LuaState, int Index, const ANSICHAR* Key)
 {
@@ -461,7 +547,7 @@ static void LuaOverridePrint(lua_State* LuaState)
 	lua_pop(LuaState, 1);
 }
 
-static void* LuaAlloc(void *Ud, void *Ptr, size_t OldSize, size_t NewSize) 
+static void* LuaAlloc(void *Ud, void *Ptr, size_t OldSize, size_t NewSize)
 {
 	if (NewSize != 0)
 	{
@@ -481,7 +567,7 @@ static int32 LuaPanic(lua_State *lua_State)
 }
 
 
-static lua_State* LuaNewState() 
+static lua_State* LuaNewState()
 {
 	lua_State* LuaState = lua_newstate(LuaAlloc, NULL);
 	if (LuaState)
@@ -504,17 +590,33 @@ FLuaContext::FLuaContext()
 , bHasBeginPlay(false)
 , LuaState(NULL)
 {
-};
+}
 
-bool FLuaContext::Initialize(UScriptComponent* Owner)
+FLuaContext::~FLuaContext()
 {
+	CloseLua();
+}
+
+void FLuaContext::CloseLua()
+{
+	if (LuaState)
+	{
+		lua_close(LuaState);
+		LuaState = NULL;
+	}
+}
+
+bool FLuaContext::Initialize(const FString& Code, UObject* Owner)
+{
+	CloseLua();
+
 	bool bResult = false;
 	LuaState = LuaNewState();
 	luaL_openlibs(LuaState);
 	LuaRegisterExportedClasses(LuaState);
 	LuaRegisterUnrealUtilities(LuaState);
 
-	if (luaL_loadstring(LuaState, TCHAR_TO_ANSI(*Owner->Script->SourceCode)) == 0)
+	if (luaL_loadstring(LuaState, TCHAR_TO_ANSI(*Code)) == 0)
 	{
 		// Set this pointer
 		lua_pushlightuserdata(LuaState, Owner);
@@ -530,7 +632,7 @@ bool FLuaContext::Initialize(UScriptComponent* Owner)
 	}
 	else
 	{
-		UE_LOG(LogScriptPlugin, Warning, TEXT("Cannot load Lua script %s for %s: %s"), *Owner->Script->GetName(), *Owner->GetPathName(), ANSI_TO_TCHAR(lua_tostring(LuaState, -1)));
+		UE_LOG(LogScriptPlugin, Warning, TEXT("Cannot load Lua script for %s: %s"), *Owner->GetPathName(), ANSI_TO_TCHAR(lua_tostring(LuaState, -1)));
 	}
 	return bResult;
 }
@@ -559,8 +661,7 @@ void FLuaContext::Destroy()
 			FLuaUtils::CallFunction(LuaState, "Destroy");
 		}
 
-		lua_close(LuaState);
-		LuaState = NULL;
+		CloseLua();
 	}
 }
 
@@ -569,5 +670,180 @@ bool FLuaContext::CanTick()
 	return bHasTick;
 }
 
+bool FLuaContext::CallFunction(const FString& FunctionName)
+{
+	check(LuaState);
+
+	bool bSuccess = FLuaUtils::DoesFunctionExist(LuaState, TCHAR_TO_ANSI(*FunctionName));
+	if (bSuccess)
+	{
+		bSuccess = FLuaUtils::CallFunction(LuaState, TCHAR_TO_ANSI(*FunctionName));
+	}
+	else
+	{
+		UE_LOG(LogScriptPlugin, Warning, TEXT("Failed to call function '%s' "), *FunctionName);
+	}
+
+	return bSuccess;
+}
+
+bool FLuaContext::SetFloatVariable(const FString& GlobalVariable, float NewValue)
+{
+	check(LuaState);
+
+	bool bExists = false;
+
+	// Does the variable exist, and is it a number?
+	lua_getglobal(LuaState, TCHAR_TO_ANSI(*GlobalVariable));
+	if (lua_isnumber(LuaState, lua_gettop(LuaState)))
+	{
+		bExists = true;
+	}
+
+	// Pop after getglobal, since no other lua function is popping the stack
+	lua_pop(LuaState, 1);
+
+	if (bExists)
+	{
+		lua_pushnumber(LuaState, NewValue);
+		lua_setglobal(LuaState, TCHAR_TO_ANSI(*GlobalVariable)); // setglobal pops the stack for us
+	}
+	else
+	{
+		UE_LOG(LogScriptPlugin, Warning, TEXT(" [LUA] Global '%s' is NaN. "), *GlobalVariable);
+	}
+
+	return bExists;
+}
+
+float FLuaContext::GetFloatVariable(const FString& GlobalVariable)
+{
+	check(LuaState);
+
+	float Results = 0;
+	lua_getglobal(LuaState, TCHAR_TO_ANSI(*GlobalVariable));
+
+	if (lua_isnumber(LuaState, lua_gettop(LuaState)))
+	{
+		Results = (float)lua_tonumber(LuaState, -1); // -1 = top of stack
+	}
+	else
+	{
+		UE_LOG(LogScriptPlugin, Warning, TEXT(" [LUA] Global '%s' is NaN. "), *GlobalVariable);
+	}
+
+	lua_pop(LuaState, 1);
+	return Results;
+}
+
+bool FLuaContext::SetStringVariable(const FString& GlobalVariable, const FString& NewValue)
+{
+	check(LuaState);
+
+	bool bExists = false;
+
+	lua_getglobal(LuaState, TCHAR_TO_ANSI(*GlobalVariable));
+	if (lua_isstring(LuaState, lua_gettop(LuaState)))
+	{
+		bExists = true;
+	}
+
+	lua_pop(LuaState, 1);
+
+	if (bExists)
+	{
+		lua_pushstring(LuaState, TCHAR_TO_ANSI(*GlobalVariable));
+		lua_setglobal(LuaState, TCHAR_TO_ANSI(*GlobalVariable));
+	}
+
+	UE_CLOG(!bExists, LogScriptPlugin, Warning, TEXT(" [LUA] Global '%s' is not a String. "), *GlobalVariable);
+	return bExists;
+}
+
+FString FLuaContext::GetStringVariable(const FString& GlobalVariable)
+{
+	check(LuaState);
+
+	FString Results = "";
+	lua_getglobal(LuaState, TCHAR_TO_ANSI(*GlobalVariable));
+
+	if (lua_isstring(LuaState, lua_gettop(LuaState)))
+	{
+		Results = (FString)lua_tostring(LuaState, -1);
+	}
+	else
+	{
+		UE_LOG(LogScriptPlugin, Warning, TEXT(" [LUA] Global '%s' is not a String. "), *GlobalVariable);
+	}
+
+	lua_pop(LuaState, 1);
+	return Results;
+}
+
+void FLuaContext::InvokeScriptFunction(FFrame& Stack, RESULT_DECL)
+{
+	P_FINISH;
+	CallFunction(Stack.CurrentNativeFunction->GetName());
+}
+
+#if WITH_EDITOR
+void FLuaContext::GetScriptDefinedFields(TArray<FScriptField>& OutFields)
+{
+	check(LuaState);
+
+	lua_pushglobaltable(LuaState);
+	lua_pushnil(LuaState);
+
+	while (lua_next(LuaState, -2) != 0)
+	{
+		int32 KeyType = lua_type(LuaState, -2);
+		if (KeyType != LUA_TSTRING)
+		{
+			lua_pop(LuaState, 1);
+			continue;
+		}
+
+		int32 ValueType = lua_type(LuaState, -1);
+		FString KeyName(ANSI_TO_TCHAR(lua_tostring(LuaState, -2)));
+		if (KeyName.Len() == 0 || KeyName.StartsWith("_") || KeyName == TEXT("this"))
+		{
+			lua_pop(LuaState, 1);
+			continue;
+		}
+
+		FScriptField PropertyInfo;
+		switch (ValueType)
+		{
+			case LUA_TNUMBER:
+				PropertyInfo.Class = UFloatProperty::StaticClass();
+				break;
+			case LUA_TSTRING:
+				PropertyInfo.Class = UStrProperty::StaticClass();
+				break;
+			case LUA_TBOOLEAN:
+				PropertyInfo.Class = UBoolProperty::StaticClass();
+				break;
+			case LUA_TNIL:
+			case LUA_TUSERDATA:
+				PropertyInfo.Class = UObjectProperty::StaticClass();
+				break;
+			case LUA_TFUNCTION:
+				if (KeyName.StartsWith(TEXT("Event")))
+				{
+					PropertyInfo.Class = UFunction::StaticClass();
+				}
+				break;
+		}
+
+		if (PropertyInfo.Class)
+		{
+			PropertyInfo.Name = FName(*KeyName);
+			OutFields.Add(PropertyInfo);
+		}
+
+		lua_pop(LuaState, 1);
+	}
+}
+#endif
 
 #endif // WITH_LUA
diff --git a/Engine/Plugins/Script/ScriptPlugin/Source/ScriptPlugin/Private/LuaIntegration.h b/Engine/Plugins/Script/ScriptPlugin/Source/ScriptPlugin/Private/LuaIntegration.h
--- a/Engine/Plugins/Script/ScriptPlugin/Source/ScriptPlugin/Private/LuaIntegration.h
+++ b/Engine/Plugins/Script/ScriptPlugin/Source/ScriptPlugin/Private/LuaIntegration.h
@@ -16,17 +16,31 @@ class FLuaContext : public FScriptContextBase
 	bool bHasDestroy;
 	bool bHasBeginPlay;
 	lua_State* LuaState;
+	
+	void CloseLua();
 
 public:
 
 	FLuaContext();
+	virtual ~FLuaContext();
 
 	// FScriptContextBase interface
-	virtual bool Initialize(UScriptComponent* Owner) override;
+	virtual bool Initialize(const FString& Code, UObject* Owner) override;
 	virtual void BeginPlay() override;
 	virtual void Tick(float DeltaTime) override;
 	virtual void Destroy() override;
 	virtual bool CanTick() override;
+
+	virtual bool CallFunction(const FString&  FunctionName) override;
+	virtual bool SetFloatVariable(const FString&  GlobalVariable, float NewValue) override;
+	virtual float GetFloatVariable(const FString&  GlobalVariable) override;
+	virtual bool SetStringVariable(const FString&  GlobalVariable, const FString& NewValue) override;
+	virtual FString GetStringVariable(const FString&  GlobalVariable) override;
+	virtual void InvokeScriptFunction(FFrame& Stack, RESULT_DECL) override;
+
+#if WITH_EDITOR
+	virtual void GetScriptDefinedFields(TArray<FScriptField>& OutFields) override;
+#endif
 };
 
 struct FLuaUtils
@@ -36,6 +50,12 @@ struct FLuaUtils
 	static bool DoesFunctionExist(lua_State* LuaState, const char* FunctionName);
 };
 
+struct FLuaVector2D
+{
+	static FVector2D Get(lua_State *LuaState, int ParamIndex);
+	static void Return(lua_State *LuaState, const FVector2D& Vec);
+};
+
 struct FLuaVector
 {
 	static FVector Get(lua_State *LuaState, int ParamIndex);
@@ -54,6 +74,12 @@ struct FLuaQuat
 	static void Return(lua_State *LuaState, const FQuat& Vec);
 };
 
+struct FLuaLinearColor
+{
+	static FLinearColor Get(lua_State *LuaState, int ParamIndex);
+	static void Return(lua_State *LuaState, const FLinearColor& Vec);
+};
+
 struct FLuaTransform
 {
 	static FTransform Get(lua_State *LuaState, int ParamIndex);
diff --git a/Engine/Plugins/Script/ScriptPlugin/Source/ScriptPlugin/Private/ScriptActor.cpp b/Engine/Plugins/Script/ScriptPlugin/Source/ScriptPlugin/Private/ScriptActor.cpp
new file mode 100644
--- /dev/null
+++ b/Engine/Plugins/Script/ScriptPlugin/Source/ScriptPlugin/Private/ScriptActor.cpp
@@ -0,0 +1,94 @@
+// Copyright 1998-2014 Epic Games, Inc. All Rights Reserved. 
+#include "ScriptPluginPrivatePCH.h"
+#include "ScriptActor.h"
+
+
+//////////////////////////////////////////////////////////////////////////
+
+AScriptActor::AScriptActor(const FPostConstructInitializeProperties& PCIP)
+: Super( PCIP )
+{
+	PrimaryActorTick.bCanEverTick = true;
+}
+
+void AScriptActor::PostLoad()
+{
+	Super::PostLoad();
+}
+
+void AScriptActor::BeginDestroy()
+{
+	if (ScriptContext)
+	{
+		ScriptContext->Destroy();
+		delete ScriptContext;
+		ScriptContext = NULL;
+	}
+
+	Super::BeginDestroy();
+}
+
+void AScriptActor::OnConstruction(const FTransform& Transform)
+{
+	Super::OnConstruction(Transform);
+
+	if (!ScriptContext)
+	{
+		UScriptBlueprintGeneratedClass* ScriptClass = NULL;
+		for (UClass* MyClass = GetClass(); MyClass && !ScriptClass; MyClass = MyClass->GetSuperClass())
+		{
+			ScriptClass = Cast<UScriptBlueprintGeneratedClass>(MyClass);
+		}
+
+		bool Result = true;
+		if (ScriptClass)
+		{
+			ScriptContext = ScriptClass->CreateContext();
+		}
+
+		if (ScriptContext && GetWorld() && GetWorld()->WorldType != EWorldType::Editor)
+		{
+			if (!ScriptContext->Initialize(ScriptClass->SourceCode, this))
+			{
+				delete ScriptContext;
+				ScriptContext = NULL;
+			}
+		}
+	}
+}
+
+void AScriptActor::BeginPlay()
+{
+	Super::BeginPlay();
+
+	if (ScriptContext && GetWorld() && GetWorld()->WorldType != EWorldType::Editor)
+	{
+		ScriptContext->BeginPlay();
+		SetActorTickEnabled(ScriptContext->CanTick());
+	}
+}
+
+void AScriptActor::TickActor(float DeltaTime, ELevelTick TickType, FActorTickFunction& ThisTickFunction)
+{
+	Super::TickActor(DeltaTime, TickType, ThisTickFunction);
+	if (ScriptContext && ScriptContext->CanTick())
+	{
+		ScriptContext->Tick(DeltaTime);
+	}
+}
+
+void AScriptActor::ReceiveDestroyed()
+{
+	Super::ReceiveDestroyed();
+}
+
+void AScriptActor::ReceiveEndPlay(EEndPlayReason::Type EndPlayReason)
+{
+	Super::ReceiveEndPlay(EndPlayReason);
+}
+
+void AScriptActor::InvokeScriptFunction(FFrame& Stack, RESULT_DECL)
+{
+	check(ScriptContext);
+	ScriptContext->InvokeScriptFunction(Stack, Result);
+}
\ No newline at end of file
diff --git a/Engine/Plugins/Script/ScriptPlugin/Source/ScriptPlugin/Private/ScriptAsset.cpp b/Engine/Plugins/Script/ScriptPlugin/Source/ScriptPlugin/Private/ScriptAsset.cpp
deleted file mode 100644
diff --git a/Engine/Plugins/Script/ScriptPlugin/Source/ScriptPlugin/Private/ScriptBlueprint.cpp b/Engine/Plugins/Script/ScriptPlugin/Source/ScriptPlugin/Private/ScriptBlueprint.cpp
new file mode 100644
--- /dev/null
+++ b/Engine/Plugins/Script/ScriptPlugin/Source/ScriptPlugin/Private/ScriptBlueprint.cpp
@@ -0,0 +1,84 @@
+// Copyright 1998-2014 Epic Games, Inc. All Rights Reserved.
+#include "ScriptPluginPrivatePCH.h"
+#include "ScriptBlueprint.h"
+#if WITH_EDITOR
+#include "BlueprintEditorUtils.h"
+#endif
+
+/////////////////////////////////////////////////////
+// UScriptBlueprint
+
+UScriptBlueprint::UScriptBlueprint(const FPostConstructInitializeProperties& PCIP)
+	: Super(PCIP)
+{
+}
+
+void UScriptBlueprint::PostLoad()
+{
+	Super::PostLoad();
+}
+
+UClass* UScriptBlueprint::RegenerateClass(UClass* ClassToRegenerate, UObject* PreviousCDO, TArray<UObject*>& ObjLoaded)
+{
+	return Super::RegenerateClass(ClassToRegenerate, PreviousCDO, ObjLoaded);
+}
+
+#if WITH_EDITOR
+UClass* UScriptBlueprint::GetBlueprintClass() const
+{
+	return UScriptBlueprintGeneratedClass::StaticClass();
+}
+
+bool UScriptBlueprint::ValidateGeneratedClass(const UClass* InClass)
+{
+	bool Result = Super::ValidateGeneratedClass(InClass);
+
+	const UScriptBlueprintGeneratedClass* GeneratedClass = Cast<const UScriptBlueprintGeneratedClass>(InClass);
+	if ( !ensure(GeneratedClass) )
+	{
+		return false;
+	}
+	const UScriptBlueprint* Blueprint = Cast<UScriptBlueprint>(GetBlueprintFromClass(GeneratedClass));
+	if ( !ensure(Blueprint) )
+	{
+		return false;
+	}
+
+	return Result;
+}
+
+void UScriptBlueprint::UpdateScriptStatus()
+{
+	FDateTime OldTimeStamp;
+	bool bCodeDirty = !FDateTime::Parse(SourceFileTimestamp, OldTimeStamp);
+	
+	if (!bCodeDirty)
+	{
+		FDateTime TimeStamp = IFileManager::Get().GetTimeStamp(*SourceFilePath);
+		bCodeDirty = TimeStamp > OldTimeStamp;
+	}
+
+	if (bCodeDirty)
+	{
+		FBlueprintEditorUtils::MarkBlueprintAsModified(this);
+	}
+}
+
+void UScriptBlueprint::UpdateSourceCodeIfChanged()
+{
+	FDateTime OldTimeStamp;	
+	bool bCodeDirty = !FDateTime::Parse(SourceFileTimestamp, OldTimeStamp);
+	FDateTime TimeStamp = IFileManager::Get().GetTimeStamp(*SourceFilePath);
+	bCodeDirty = bCodeDirty || (TimeStamp > OldTimeStamp);
+
+	if (bCodeDirty)
+	{
+		FString NewScript;
+		if (FFileHelper::LoadFileToString(NewScript, *SourceFilePath))
+		{
+			SourceCode = NewScript;
+			SourceFileTimestamp = TimeStamp.ToString();
+		}
+	}
+}
+#endif
\ No newline at end of file
diff --git a/Engine/Plugins/Script/ScriptPlugin/Source/ScriptPlugin/Private/ScriptBlueprintGeneratedClass.cpp b/Engine/Plugins/Script/ScriptPlugin/Source/ScriptPlugin/Private/ScriptBlueprintGeneratedClass.cpp
new file mode 100644
--- /dev/null
+++ b/Engine/Plugins/Script/ScriptPlugin/Source/ScriptPlugin/Private/ScriptBlueprintGeneratedClass.cpp
@@ -0,0 +1,61 @@
+// Copyright 1998-2014 Epic Games, Inc. All Rights Reserved.
+#include "ScriptPluginPrivatePCH.h"
+#include "ScriptBlueprintGeneratedClass.h"
+#if WITH_LUA
+#include "LuaIntegration.h"
+#endif
+
+/////////////////////////////////////////////////////
+// UScriptBlueprintGeneratedClass
+
+UScriptBlueprintGeneratedClass::UScriptBlueprintGeneratedClass(const FPostConstructInitializeProperties& PCIP)
+	: Super(PCIP)
+{
+}
+
+void UScriptBlueprintGeneratedClass::PostInitProperties()
+{
+	Super::PostInitProperties();
+}
+
+void UScriptBlueprintGeneratedClass::Link(FArchive& Ar, bool bRelinkExistingProperties)
+{
+	Super::Link(Ar, bRelinkExistingProperties);
+}
+
+FScriptContextBase* UScriptBlueprintGeneratedClass::CreateContext()
+{
+	FScriptContextBase* NewContext = NULL;
+#if WITH_LUA
+	NewContext = new FLuaContext();
+#endif
+	return NewContext;
+}
+
+void UScriptBlueprintGeneratedClass::AddUniqueNativeFunction(const FName& InName, Native InPointer)
+{
+	// Find the function in the class's native function lookup table.
+	for (int32 FunctionIndex = 0; FunctionIndex < NativeFunctionLookupTable.Num(); ++FunctionIndex)
+	{
+		FNativeFunctionLookup& NativeFunctionLookup = NativeFunctionLookupTable[FunctionIndex];
+		if (NativeFunctionLookup.Name == InName)
+		{
+			return;
+		}
+	}
+	new(NativeFunctionLookupTable)FNativeFunctionLookup(InName, InPointer);
+}
+
+void UScriptBlueprintGeneratedClass::RemoveNativeFunction(const FName& InName)
+{
+	// Find the function in the class's native function lookup table.
+	for (int32 FunctionIndex = 0; FunctionIndex < NativeFunctionLookupTable.Num(); ++FunctionIndex)
+	{
+		FNativeFunctionLookup& NativeFunctionLookup = NativeFunctionLookupTable[FunctionIndex];
+		if (NativeFunctionLookup.Name == InName)
+		{
+			NativeFunctionLookupTable.RemoveAt(FunctionIndex);
+			return;
+		}
+	}
+}
\ No newline at end of file
diff --git a/Engine/Plugins/Script/ScriptPlugin/Source/ScriptPlugin/Private/ScriptComponent.cpp b/Engine/Plugins/Script/ScriptPlugin/Source/ScriptPlugin/Private/ScriptComponent.cpp
--- a/Engine/Plugins/Script/ScriptPlugin/Source/ScriptPlugin/Private/ScriptComponent.cpp
+++ b/Engine/Plugins/Script/ScriptPlugin/Source/ScriptPlugin/Private/ScriptComponent.cpp
@@ -1,12 +1,11 @@
 // Copyright 1998-2014 Epic Games, Inc. All Rights Reserved. 
 #include "ScriptPluginPrivatePCH.h"
 
-
 //////////////////////////////////////////////////////////////////////////
 
 UScriptComponent::UScriptComponent(const FPostConstructInitializeProperties& PCIP)
-	: Super( PCIP )
-{	
+	: Super(PCIP)
+{
 	PrimaryComponentTick.bCanEverTick = true;
 	bTickInEditor = false;
 	bAutoActivate = true;
@@ -18,21 +17,20 @@ UScriptComponent::UScriptComponent(const FPostConstructInitializeProperties& PCI
 void UScriptComponent::OnRegister()
 {
 	Super::OnRegister();
-	if (Script && GetWorld() && GetWorld()->WorldType != EWorldType::Editor)
+
+	UScriptBlueprintGeneratedClass* ScriptClass = NULL;
+	for (UClass* MyClass = GetClass(); MyClass && !ScriptClass; MyClass = MyClass->GetSuperClass())
 	{
-#if WITH_EDITOR
-		Script->UpdateAndCompile();
-#endif
-#if WITH_LUA
-		Context = new FLuaContext();
-#else
-		// @todo Create context here
-#endif
+		ScriptClass = Cast<UScriptBlueprintGeneratedClass>(MyClass);
+	}
 
+	if (ScriptClass && GetWorld() && GetWorld()->WorldType != EWorldType::Editor)
+	{
+		Context = ScriptClass->CreateContext();
 		if (Context)
 		{
 			bool bDoNotTick = true;
-			if (Context->Initialize(this))
+			if (Context->Initialize(ScriptClass->SourceCode, this))
 			{
 				bDoNotTick = !Context->CanTick();
 			}
@@ -73,4 +71,14 @@ void UScriptComponent::OnUnregister()
 	}
 
 	Super::OnUnregister();
-}
\ No newline at end of file
+}
+
+bool UScriptComponent::CallScriptFunction(FString FunctionName)
+{
+	bool bSuccess = false;
+	if (Context)
+	{
+		bSuccess = Context->CallFunction(FunctionName);
+	}
+	return bSuccess;
+}
diff --git a/Engine/Plugins/Script/ScriptPlugin/Source/ScriptPlugin/Private/ScriptFunction.cpp b/Engine/Plugins/Script/ScriptPlugin/Source/ScriptPlugin/Private/ScriptFunction.cpp
new file mode 100644
--- /dev/null
+++ b/Engine/Plugins/Script/ScriptPlugin/Source/ScriptPlugin/Private/ScriptFunction.cpp
@@ -0,0 +1,18 @@
+// Copyright 1998-2014 Epic Games, Inc. All Rights Reserved. 
+#include "ScriptPluginPrivatePCH.h"
+#include "ScriptFunction.h"
+#include "ScriptBlueprintGeneratedClass.h"
+#include "ScriptActor.h"
+
+//////////////////////////////////////////////////////////////////////////
+
+UScriptFunction::UScriptFunction(const FPostConstructInitializeProperties& PCIP)
+: Super(PCIP, NULL, FUNC_Public | FUNC_BlueprintCallable | FUNC_Native, 65535)
+{
+	UScriptBlueprintGeneratedClass* ScriptClass = Cast<UScriptBlueprintGeneratedClass>(GetOuter());
+	if (ScriptClass)
+	{
+		ScriptClass->AddUniqueNativeFunction(GetFName(), (Native)&AScriptActor::InvokeScriptFunction);
+		SetNativeFunc((Native)&AScriptActor::InvokeScriptFunction);
+	}
+}
diff --git a/Engine/Plugins/Script/ScriptPlugin/Source/ScriptPlugin/Private/ScriptPluginPrivatePCH.h b/Engine/Plugins/Script/ScriptPlugin/Source/ScriptPlugin/Private/ScriptPluginPrivatePCH.h
--- a/Engine/Plugins/Script/ScriptPlugin/Source/ScriptPlugin/Private/ScriptPluginPrivatePCH.h
+++ b/Engine/Plugins/Script/ScriptPlugin/Source/ScriptPlugin/Private/ScriptPluginPrivatePCH.h
@@ -6,9 +6,12 @@
 #include "ModuleManager.h"
 #include "Engine.h"
 #include "ScriptComponent.h"
-#include "ScriptAsset.h"
 #include "ScriptTestActor.h"
 
+#if WITH_EDITOR
+#include "UnrealEd.h"
+#endif
+
 #if WITH_LUA
 #include "LuaIntegration.h"
 #endif // WITH_LUA
diff --git a/Engine/Plugins/Script/ScriptPlugin/Source/ScriptPlugin/ScriptPlugin.Build.cs b/Engine/Plugins/Script/ScriptPlugin/Source/ScriptPlugin/ScriptPlugin.Build.cs
--- a/Engine/Plugins/Script/ScriptPlugin/Source/ScriptPlugin/ScriptPlugin.Build.cs
+++ b/Engine/Plugins/Script/ScriptPlugin/Source/ScriptPlugin/ScriptPlugin.Build.cs
@@ -27,11 +27,23 @@ namespace UnrealBuildTool.Rules
 					"Core",
 					"CoreUObject",
 					"Engine",
-					"InputCore"
+					"InputCore",
 					// ... add other public dependencies that you statically link with here ...
 				}
 				);
 
+			if (UEBuildConfiguration.bBuildEditor == true)
+			{
+
+				PublicDependencyModuleNames.AddRange(
+					new string[] 
+					{
+						"UnrealEd", 
+					}
+				);
+
+			}
+
 			PrivateDependencyModuleNames.AddRange(
 				new string[]
 				{
diff --git a/Engine/Source/Runtime/CoreUObject/Private/UObject/Class.cpp b/Engine/Source/Runtime/CoreUObject/Private/UObject/Class.cpp
--- a/Engine/Source/Runtime/CoreUObject/Private/UObject/Class.cpp
+++ b/Engine/Source/Runtime/CoreUObject/Private/UObject/Class.cpp
@@ -3347,11 +3347,6 @@ bool UClass::HotReloadPrivateStaticClass(
 }
 #endif
 
-/**
-* Add a native function to the internal native function table
-* @param	InName							name of the function
-* @param	InPointer						pointer to the function
-**/
 void UClass::AddNativeFunction(const ANSICHAR* InName,Native InPointer)
 {
 	FName InFName(InName);
@@ -3666,6 +3661,7 @@ void UFunction::Invoke(UObject* Obj, FFrame& Stack, RESULT_DECL)
 		Obj = (UObject*)Obj->GetInterfaceAddress(OuterClass);
 	}
 
+	TGuardValue<UFunction*> NativeFuncGuard(Stack.CurrentNativeFunction, this);
 	return (Obj->*Func)(Stack, Result);
 }
 
diff --git a/Engine/Source/Runtime/CoreUObject/Public/UObject/Class.h b/Engine/Source/Runtime/CoreUObject/Public/UObject/Class.h
--- a/Engine/Source/Runtime/CoreUObject/Public/UObject/Class.h
+++ b/Engine/Source/Runtime/CoreUObject/Public/UObject/Class.h
@@ -1136,9 +1136,9 @@ public:
 //
 // Reflection data for a replicated or Kismet callable function.
 //
-class UFunction : public UStruct
+class COREUOBJECT_API UFunction : public UStruct
 {
-	DECLARE_CASTED_CLASS_INTRINSIC_WITH_API(UFunction,UStruct,0,CoreUObject,CASTCLASS_UFunction,COREUOBJECT_API)
+	DECLARE_CASTED_CLASS_INTRINSIC(UFunction, UStruct, 0, CoreUObject, CASTCLASS_UFunction)
 	DECLARE_WITHIN(UClass)
 public:
 	// Persistent variables.
@@ -1191,7 +1191,7 @@ public:
 	void Invoke(UObject* Obj, FFrame& Stack, RESULT_DECL);
 
 	// Constructors.
-	COREUOBJECT_API explicit UFunction(const class FPostConstructInitializeProperties& PCIP, UFunction* InSuperFunction, uint32 InFunctionFlags = 0, uint16 InRepOffset = 0, SIZE_T ParamsSize = 0 );
+	explicit UFunction(const class FPostConstructInitializeProperties& PCIP, UFunction* InSuperFunction, uint32 InFunctionFlags = 0, uint16 InRepOffset = 0, SIZE_T ParamsSize = 0 );
 
 	void InitializeDerivedMembers();
 
@@ -1211,7 +1211,7 @@ public:
 		checkSlow(!SuperStruct||SuperStruct->IsA<UFunction>());
 		return (UFunction*)SuperStruct;
 	}
-	COREUOBJECT_API UProperty* GetReturnProperty();
+	UProperty* GetReturnProperty();
 
 	/**
 	 * Used to safely check whether the passed in flag is set.
@@ -1256,7 +1256,7 @@ public:
 	 *
 	 * @return	true if function signatures are compatible.
 	 */
-	COREUOBJECT_API bool IsSignatureCompatibleWith(const UFunction* OtherFunction) const;
+	bool IsSignatureCompatibleWith(const UFunction* OtherFunction) const;
 
 	/**
 	 * Determines if two functions have an identical signature (note: currently doesn't allow
@@ -1268,7 +1268,7 @@ public:
 	 *
 	 * @return	true if function signatures are compatible.
 	 */
-	COREUOBJECT_API bool IsSignatureCompatibleWith(const UFunction* OtherFunction, uint64 IgnoreFlags) const;
+	bool IsSignatureCompatibleWith(const UFunction* OtherFunction, uint64 IgnoreFlags) const;
 };
 
 /*-----------------------------------------------------------------------------
diff --git a/Engine/Source/Runtime/CoreUObject/Public/UObject/Stack.h b/Engine/Source/Runtime/CoreUObject/Public/UObject/Stack.h
--- a/Engine/Source/Runtime/CoreUObject/Public/UObject/Stack.h
+++ b/Engine/Source/Runtime/CoreUObject/Public/UObject/Stack.h
@@ -96,6 +96,9 @@ public:
 
 	/** If a class is compiled in then this is set to the property chain for compiled-in functions. In that case, we follow the links to setup the args instead of executing by code. */
 	UField* PropertyChainForCompiledIn;
+
+	/** Currently executed native function */
+	UFunction* CurrentNativeFunction;
 public:
 
 	// Constructors.
@@ -175,6 +178,7 @@ inline FFrame::FFrame( UObject* InObject, UFunction* InNode, void* InLocals, FFr
 	, PreviousFrame(InPreviousFrame)
 	, OutParms(NULL)
 	, PropertyChainForCompiledIn(InPropertyChainForCompiledIn)
+	, CurrentNativeFunction(NULL)
 {
 #if !(UE_BUILD_SHIPPING || UE_BUILD_TEST)
 	FScriptTraceStackNode StackNode(InNode->GetOuter()->GetFName(), InNode->GetFName());
