ScriptPlugin: Keep script and generated class properties in sync

diff --git a/Engine/Plugins/Script/ScriptEditorPlugin/Source/ScriptEditorPlugin/Private/ScriptBlueprintCompiler.cpp b/Engine/Plugins/Script/ScriptEditorPlugin/Source/ScriptEditorPlugin/Private/ScriptBlueprintCompiler.cpp
--- a/Engine/Plugins/Script/ScriptEditorPlugin/Source/ScriptEditorPlugin/Private/ScriptBlueprintCompiler.cpp
+++ b/Engine/Plugins/Script/ScriptEditorPlugin/Source/ScriptEditorPlugin/Private/ScriptBlueprintCompiler.cpp
@@ -38,6 +38,8 @@ void FScriptBlueprintCompiler::CreateClassVariablesFromBlueprint()
 	Super::CreateClassVariablesFromBlueprint();
 
 	UScriptBlueprint* Blueprint = ScriptBlueprint();
+	UScriptBlueprintGeneratedClass* NewScripClass = CastChecked<UScriptBlueprintGeneratedClass>(NewClass);
+	NewScripClass->ScriptProperties.Empty();
 
 	for (auto& Field : ScriptDefinedFields)
 	{
@@ -84,7 +86,8 @@ void FScriptBlueprintCompiler::CreateClassVariablesFromBlueprint()
 				if (ScriptProperty != NULL)
 				{
 					ScriptProperty->SetMetaData(TEXT("Category"), *Blueprint->GetName());
-					ScriptProperty->SetPropertyFlags(CPF_BlueprintVisible | CPF_Edit);
+					ScriptProperty->SetPropertyFlags(CPF_BlueprintVisible | CPF_Edit);					
+					NewScripClass->ScriptProperties.Add(ScriptProperty);
 				}
 			}
 		}
diff --git a/Engine/Plugins/Script/ScriptPlugin/Source/ScriptPlugin/Classes/ScriptBlueprintGeneratedClass.h b/Engine/Plugins/Script/ScriptPlugin/Source/ScriptPlugin/Classes/ScriptBlueprintGeneratedClass.h
--- a/Engine/Plugins/Script/ScriptPlugin/Source/ScriptPlugin/Classes/ScriptBlueprintGeneratedClass.h
+++ b/Engine/Plugins/Script/ScriptPlugin/Source/ScriptPlugin/Classes/ScriptBlueprintGeneratedClass.h
@@ -5,6 +5,8 @@
 #include "Engine/BlueprintGeneratedClass.h"
 #include "ScriptBlueprintGeneratedClass.generated.h"
 
+class UScriptBlueprintGeneratedClass;
+
 /**
 * Script-defined field (variable or function)
 */
@@ -30,6 +32,8 @@ struct FScriptField
 class FScriptContextBase
 {
 public:
+	virtual ~FScriptContextBase() {}
+
 	/**
 	 * Initializes script context given script code
 	 * @param Code Script code
@@ -63,10 +67,18 @@ public:
 	*/
 	virtual bool CallFunction(const FString&  FunctionName) = 0;
 
-	virtual bool SetFloatVariable(const FString&  GlobalVariable, float NewValue) = 0;
-	virtual float GetFloatVariable(const FString&  GlobalVariable) = 0;
-	virtual bool SetStringVariable(const FString&  GlobalVariable, const FString& NewValue) = 0;
-	virtual FString GetStringVariable(const FString&  GlobalVariable) = 0;
+	// Property accessors
+
+	virtual bool SetFloatProperty(const FString& PropertyName, float NewValue) = 0;
+	virtual float GetFloatProperty(const FString& PropertyName) = 0;
+	virtual bool SetIntProperty(const FString& PropertyName, int32 NewValue) = 0;
+	virtual int32 GetIntProperty(const FString& PropertyName) = 0;
+	virtual bool SetObjectProperty(const FString& PropertyName, UObject* NewValue) = 0;
+	virtual UObject* GetObjectProperty(const FString& PropertyName) = 0;
+	virtual bool SetBoolProperty(const FString& PropertyName, bool NewValue) = 0;
+	virtual bool GetBoolProperty(const FString& PropertyName) = 0;
+	virtual bool SetStringProperty(const FString& PropertyName, const FString& NewValue) = 0;
+	virtual FString GetStringProperty(const FString& PropertyName) = 0;
 
 	/**
 	* Invokes script function from Blueprint code
@@ -75,12 +87,22 @@ public:
 
 #if WITH_EDITOR
 	/**
-	* Returns a list of exported fields from the script (member variables and functions).
+	* Returns a list of exported fields from script (member variables and functions).
 	*/
 	virtual void GetScriptDefinedFields(TArray<FScriptField>& OutFields) = 0;
 #endif
 
-	virtual ~FScriptContextBase() {}
+	// Utilities
+
+	/**
+	* Pushes all property values from class instance to script.
+	*/
+	virtual void PushScriptPropertyValues(UScriptBlueprintGeneratedClass* Class, const UObject* Obj);
+
+	/**
+	* Fetches all property values from script to class instance.
+	*/
+	virtual void FetchScriptPropertyValues(UScriptBlueprintGeneratedClass* Class, UObject* Obj);	
 };
 
 /**
@@ -99,8 +121,13 @@ class SCRIPTPLUGIN_API UScriptBlueprintGeneratedClass : public UBlueprintGenerat
 	UPROPERTY()
 	FString SourceCode;
 
+	/** Script-generated properties */
+	UPROPERTY()
+	TArray<UProperty*> ScriptProperties;
+
 	virtual void PostInitProperties() override;
 	virtual void Link(FArchive& Ar, bool bRelinkExistingProperties) override;
+	virtual void PurgeClass(bool bRecompilingOnLoad) override;
 
 	/**
 	 * Creates a script context object
@@ -119,4 +146,18 @@ class SCRIPTPLUGIN_API UScriptBlueprintGeneratedClass : public UBlueprintGenerat
 	* @param InName Name of the native function
 	*/
 	void RemoveNativeFunction(const FName& InName);
+
+	/**
+	* Gets UScriptBlueprintGeneratedClass from class hierarchy
+	* @return UScriptBlueprintGeneratedClass or NULL
+	*/
+	FORCEINLINE static UScriptBlueprintGeneratedClass* GetScriptGeneratedClass(UClass* InClass)
+	{
+		UScriptBlueprintGeneratedClass* ScriptClass = NULL;
+		for (UClass* MyClass = InClass; MyClass && !ScriptClass; MyClass = MyClass->GetSuperClass())
+		{
+			ScriptClass = Cast<UScriptBlueprintGeneratedClass>(MyClass);
+		}
+		return ScriptClass;
+	}
 };
\ No newline at end of file
diff --git a/Engine/Plugins/Script/ScriptPlugin/Source/ScriptPlugin/Private/LuaIntegration.cpp b/Engine/Plugins/Script/ScriptPlugin/Source/ScriptPlugin/Private/LuaIntegration.cpp
--- a/Engine/Plugins/Script/ScriptPlugin/Source/ScriptPlugin/Private/LuaIntegration.cpp
+++ b/Engine/Plugins/Script/ScriptPlugin/Source/ScriptPlugin/Private/LuaIntegration.cpp
@@ -687,14 +687,14 @@ bool FLuaContext::CallFunction(const FString& FunctionName)
 	return bSuccess;
 }
 
-bool FLuaContext::SetFloatVariable(const FString& GlobalVariable, float NewValue)
+bool FLuaContext::SetFloatProperty(const FString& PropertyName, float NewValue)
 {
 	check(LuaState);
 
 	bool bExists = false;
 
-	// Does the variable exist, and is it a number?
-	lua_getglobal(LuaState, TCHAR_TO_ANSI(*GlobalVariable));
+	// Does the Property exist, and is it a number?
+	lua_getglobal(LuaState, TCHAR_TO_ANSI(*PropertyName));
 	if (lua_isnumber(LuaState, lua_gettop(LuaState)))
 	{
 		bExists = true;
@@ -706,22 +706,22 @@ bool FLuaContext::SetFloatVariable(const FString& GlobalVariable, float NewValue
 	if (bExists)
 	{
 		lua_pushnumber(LuaState, NewValue);
-		lua_setglobal(LuaState, TCHAR_TO_ANSI(*GlobalVariable)); // setglobal pops the stack for us
+		lua_setglobal(LuaState, TCHAR_TO_ANSI(*PropertyName)); // setglobal pops the stack for us
 	}
 	else
 	{
-		UE_LOG(LogScriptPlugin, Warning, TEXT(" [LUA] Global '%s' is NaN. "), *GlobalVariable);
+		UE_LOG(LogScriptPlugin, Warning, TEXT(" [LUA] Global '%s' is NaN. "), *PropertyName);
 	}
 
 	return bExists;
 }
 
-float FLuaContext::GetFloatVariable(const FString& GlobalVariable)
+float FLuaContext::GetFloatProperty(const FString& PropertyName)
 {
 	check(LuaState);
 
 	float Results = 0;
-	lua_getglobal(LuaState, TCHAR_TO_ANSI(*GlobalVariable));
+	lua_getglobal(LuaState, TCHAR_TO_ANSI(*PropertyName));
 
 	if (lua_isnumber(LuaState, lua_gettop(LuaState)))
 	{
@@ -729,20 +729,128 @@ float FLuaContext::GetFloatVariable(const FString& GlobalVariable)
 	}
 	else
 	{
-		UE_LOG(LogScriptPlugin, Warning, TEXT(" [LUA] Global '%s' is NaN. "), *GlobalVariable);
+		UE_LOG(LogScriptPlugin, Warning, TEXT(" [LUA] Global '%s' is NaN. "), *PropertyName);
 	}
 
 	lua_pop(LuaState, 1);
 	return Results;
 }
 
-bool FLuaContext::SetStringVariable(const FString& GlobalVariable, const FString& NewValue)
+bool FLuaContext::SetIntProperty(const FString& PropertyName, int32 NewValue)
+{
+	return SetFloatProperty(PropertyName, (float)NewValue);
+}
+
+int32 FLuaContext::GetIntProperty(const FString& PropertyName)
+{
+	return FMath::TruncToInt(GetFloatProperty(PropertyName));
+}
+
+bool FLuaContext::SetObjectProperty(const FString& PropertyName, UObject* NewValue)
+{
+	check(LuaState);
+
+	bool bExists = false;
+
+	// Does the Property exist, and is it a number?
+	lua_getglobal(LuaState, TCHAR_TO_ANSI(*PropertyName));
+	if (lua_islightuserdata(LuaState, lua_gettop(LuaState)))
+	{
+		bExists = true;
+	}
+
+	// Pop after getglobal, since no other lua function is popping the stack
+	lua_pop(LuaState, 1);
+
+	if (bExists)
+	{
+		lua_pushlightuserdata(LuaState, NewValue);
+		lua_setglobal(LuaState, TCHAR_TO_ANSI(*PropertyName)); // setglobal pops the stack for us
+	}
+	else
+	{
+		UE_LOG(LogScriptPlugin, Warning, TEXT(" [LUA] Global '%s' is NaN. "), *PropertyName);
+	}
+
+	return bExists;
+}
+
+UObject* FLuaContext::GetObjectProperty(const FString& PropertyName)
+{
+	check(LuaState);
+
+	UObject* Results = NULL;
+	lua_getglobal(LuaState, TCHAR_TO_ANSI(*PropertyName));
+
+	if (lua_islightuserdata(LuaState, lua_gettop(LuaState)))
+	{
+		Results = (UObject*)lua_touserdata(LuaState, -1); // -1 = top of stack
+	}
+	else
+	{
+		UE_LOG(LogScriptPlugin, Warning, TEXT(" [LUA] Global '%s' is NaN. "), *PropertyName);
+	}
+
+	lua_pop(LuaState, 1);
+	return Results;
+}
+
+bool FLuaContext::SetBoolProperty(const FString& PropertyName, bool NewValue)
+{
+	check(LuaState);
+
+	bool bExists = false;
+
+	// Does the Property exist, and is it a number?
+	lua_getglobal(LuaState, TCHAR_TO_ANSI(*PropertyName));
+	if (lua_isboolean(LuaState, lua_gettop(LuaState)))
+	{
+		bExists = true;
+	}
+
+	// Pop after getglobal, since no other lua function is popping the stack
+	lua_pop(LuaState, 1);
+
+	if (bExists)
+	{
+		lua_pushboolean(LuaState, (int32)NewValue);
+		lua_setglobal(LuaState, TCHAR_TO_ANSI(*PropertyName)); // setglobal pops the stack for us
+	}
+	else
+	{
+		UE_LOG(LogScriptPlugin, Warning, TEXT(" [LUA] Global '%s' is NaN. "), *PropertyName);
+	}
+
+	return bExists;
+}
+
+bool FLuaContext::GetBoolProperty(const FString& PropertyName)
+{
+	check(LuaState);
+
+	bool Results = false;
+	lua_getglobal(LuaState, TCHAR_TO_ANSI(*PropertyName));
+
+	if (lua_isboolean(LuaState, lua_gettop(LuaState)))
+	{
+		Results = !!lua_toboolean(LuaState, -1); // -1 = top of stack
+	}
+	else
+	{
+		UE_LOG(LogScriptPlugin, Warning, TEXT(" [LUA] Global '%s' is NaN. "), *PropertyName);
+	}
+
+	lua_pop(LuaState, 1);
+	return Results;
+}
+
+bool FLuaContext::SetStringProperty(const FString& PropertyName, const FString& NewValue)
 {
 	check(LuaState);
 
 	bool bExists = false;
 
-	lua_getglobal(LuaState, TCHAR_TO_ANSI(*GlobalVariable));
+	lua_getglobal(LuaState, TCHAR_TO_ANSI(*PropertyName));
 	if (lua_isstring(LuaState, lua_gettop(LuaState)))
 	{
 		bExists = true;
@@ -752,20 +860,20 @@ bool FLuaContext::SetStringVariable(const FString& GlobalVariable, const FString
 
 	if (bExists)
 	{
-		lua_pushstring(LuaState, TCHAR_TO_ANSI(*GlobalVariable));
-		lua_setglobal(LuaState, TCHAR_TO_ANSI(*GlobalVariable));
+		lua_pushstring(LuaState, TCHAR_TO_ANSI(*PropertyName));
+		lua_setglobal(LuaState, TCHAR_TO_ANSI(*PropertyName));
 	}
 
-	UE_CLOG(!bExists, LogScriptPlugin, Warning, TEXT(" [LUA] Global '%s' is not a String. "), *GlobalVariable);
+	UE_CLOG(!bExists, LogScriptPlugin, Warning, TEXT(" [LUA] Global '%s' is not a String. "), *PropertyName);
 	return bExists;
 }
 
-FString FLuaContext::GetStringVariable(const FString& GlobalVariable)
+FString FLuaContext::GetStringProperty(const FString& PropertyName)
 {
 	check(LuaState);
 
 	FString Results = "";
-	lua_getglobal(LuaState, TCHAR_TO_ANSI(*GlobalVariable));
+	lua_getglobal(LuaState, TCHAR_TO_ANSI(*PropertyName));
 
 	if (lua_isstring(LuaState, lua_gettop(LuaState)))
 	{
@@ -773,7 +881,7 @@ FString FLuaContext::GetStringVariable(const FString& GlobalVariable)
 	}
 	else
 	{
-		UE_LOG(LogScriptPlugin, Warning, TEXT(" [LUA] Global '%s' is not a String. "), *GlobalVariable);
+		UE_LOG(LogScriptPlugin, Warning, TEXT(" [LUA] Global '%s' is not a String. "), *PropertyName);
 	}
 
 	lua_pop(LuaState, 1);
diff --git a/Engine/Plugins/Script/ScriptPlugin/Source/ScriptPlugin/Private/LuaIntegration.h b/Engine/Plugins/Script/ScriptPlugin/Source/ScriptPlugin/Private/LuaIntegration.h
--- a/Engine/Plugins/Script/ScriptPlugin/Source/ScriptPlugin/Private/LuaIntegration.h
+++ b/Engine/Plugins/Script/ScriptPlugin/Source/ScriptPlugin/Private/LuaIntegration.h
@@ -31,11 +31,18 @@ public:
 	virtual void Destroy() override;
 	virtual bool CanTick() override;
 
-	virtual bool CallFunction(const FString&  FunctionName) override;
-	virtual bool SetFloatVariable(const FString&  GlobalVariable, float NewValue) override;
-	virtual float GetFloatVariable(const FString&  GlobalVariable) override;
-	virtual bool SetStringVariable(const FString&  GlobalVariable, const FString& NewValue) override;
-	virtual FString GetStringVariable(const FString&  GlobalVariable) override;
+	virtual bool SetFloatProperty(const FString& PropertyName, float NewValue) override;
+	virtual float GetFloatProperty(const FString& PropertyName) override;
+	virtual bool SetIntProperty(const FString& PropertyName, int32 NewValue) override;
+	virtual int32 GetIntProperty(const FString& PropertyName) override;
+	virtual bool SetObjectProperty(const FString& PropertyName, UObject* NewValue) override;
+	virtual UObject* GetObjectProperty(const FString& PropertyName) override;
+	virtual bool SetBoolProperty(const FString& PropertyName, bool NewValue) override;
+	virtual bool GetBoolProperty(const FString& PropertyName) override;
+	virtual bool SetStringProperty(const FString& PropertyName, const FString& NewValue) override;
+	virtual FString GetStringProperty(const FString& PropertyName) override;
+
+	virtual bool CallFunction(const FString& FunctionName) override;
 	virtual void InvokeScriptFunction(FFrame& Stack, RESULT_DECL) override;
 
 #if WITH_EDITOR
diff --git a/Engine/Plugins/Script/ScriptPlugin/Source/ScriptPlugin/Private/ScriptActor.cpp b/Engine/Plugins/Script/ScriptPlugin/Source/ScriptPlugin/Private/ScriptActor.cpp
--- a/Engine/Plugins/Script/ScriptPlugin/Source/ScriptPlugin/Private/ScriptActor.cpp
+++ b/Engine/Plugins/Script/ScriptPlugin/Source/ScriptPlugin/Private/ScriptActor.cpp
@@ -34,12 +34,7 @@ void AScriptActor::OnConstruction(const FTransform& Transform)
 
 	if (!ScriptContext)
 	{
-		UScriptBlueprintGeneratedClass* ScriptClass = NULL;
-		for (UClass* MyClass = GetClass(); MyClass && !ScriptClass; MyClass = MyClass->GetSuperClass())
-		{
-			ScriptClass = Cast<UScriptBlueprintGeneratedClass>(MyClass);
-		}
-
+		auto ScriptClass = UScriptBlueprintGeneratedClass::GetScriptGeneratedClass(GetClass());
 		bool Result = true;
 		if (ScriptClass)
 		{
@@ -53,6 +48,11 @@ void AScriptActor::OnConstruction(const FTransform& Transform)
 				delete ScriptContext;
 				ScriptContext = NULL;
 			}
+			else
+			{
+				// Push values set by CDO
+				ScriptContext->PushScriptPropertyValues(ScriptClass, this);
+			}
 		}
 	}
 }
@@ -63,7 +63,10 @@ void AScriptActor::BeginPlay()
 
 	if (ScriptContext && GetWorld() && GetWorld()->WorldType != EWorldType::Editor)
 	{
+		auto ScriptClass = UScriptBlueprintGeneratedClass::GetScriptGeneratedClass(GetClass());
+		ScriptContext->PushScriptPropertyValues(ScriptClass, this);
 		ScriptContext->BeginPlay();
+		ScriptContext->FetchScriptPropertyValues(ScriptClass, this);
 		SetActorTickEnabled(ScriptContext->CanTick());
 	}
 }
@@ -73,7 +76,10 @@ void AScriptActor::TickActor(float DeltaTime, ELevelTick TickType, FActorTickFun
 	Super::TickActor(DeltaTime, TickType, ThisTickFunction);
 	if (ScriptContext && ScriptContext->CanTick())
 	{
+		auto ScriptClass = UScriptBlueprintGeneratedClass::GetScriptGeneratedClass(GetClass());
+		ScriptContext->PushScriptPropertyValues(ScriptClass, this);
 		ScriptContext->Tick(DeltaTime);
+		ScriptContext->FetchScriptPropertyValues(ScriptClass, this);
 	}
 }
 
@@ -90,5 +96,8 @@ void AScriptActor::ReceiveEndPlay(EEndPlayReason::Type EndPlayReason)
 void AScriptActor::InvokeScriptFunction(FFrame& Stack, RESULT_DECL)
 {
 	check(ScriptContext);
+	auto ScriptClass = UScriptBlueprintGeneratedClass::GetScriptGeneratedClass(GetClass());
+	ScriptContext->PushScriptPropertyValues(ScriptClass, this);
 	ScriptContext->InvokeScriptFunction(Stack, Result);
+	ScriptContext->FetchScriptPropertyValues(ScriptClass, this);
 }
\ No newline at end of file
diff --git a/Engine/Plugins/Script/ScriptPlugin/Source/ScriptPlugin/Private/ScriptBlueprintGeneratedClass.cpp b/Engine/Plugins/Script/ScriptPlugin/Source/ScriptPlugin/Private/ScriptBlueprintGeneratedClass.cpp
--- a/Engine/Plugins/Script/ScriptPlugin/Source/ScriptPlugin/Private/ScriptBlueprintGeneratedClass.cpp
+++ b/Engine/Plugins/Script/ScriptPlugin/Source/ScriptPlugin/Private/ScriptBlueprintGeneratedClass.cpp
@@ -58,4 +58,76 @@ void UScriptBlueprintGeneratedClass::RemoveNativeFunction(const FName& InName)
 			return;
 		}
 	}
+}
+
+void UScriptBlueprintGeneratedClass::PurgeClass(bool bRecompilingOnLoad)
+{
+	Super::PurgeClass(bRecompilingOnLoad);
+	ScriptProperties.Empty();
+}
+
+void FScriptContextBase::PushScriptPropertyValues(UScriptBlueprintGeneratedClass* Class, const UObject* Obj)
+{
+	// @todo: optimize this
+	for (auto Property : Class->ScriptProperties)
+	{
+		if (UFloatProperty* FloatProperty = Cast<UFloatProperty>(Property))
+		{
+			float Value = FloatProperty->GetFloatingPointPropertyValue(Property->ContainerPtrToValuePtr<float>(Obj));
+			SetFloatProperty(Property->GetName(), Value);
+		}
+		else if (UIntProperty* IntProperty = Cast<UIntProperty>(Property))
+		{
+			int32 Value = IntProperty->GetSignedIntPropertyValue(Property->ContainerPtrToValuePtr<int32>(Obj));
+			SetIntProperty(Property->GetName(), Value);
+		}
+		else if (UBoolProperty* BoolProperty = Cast<UBoolProperty>(Property))
+		{
+			bool Value = BoolProperty->GetPropertyValue(Property->ContainerPtrToValuePtr<void>(Obj));
+			SetBoolProperty(Property->GetName(), Value);
+		}
+		else if (UObjectProperty* ObjectProperty = Cast<UObjectProperty>(Property))
+		{
+			UObject* Value = ObjectProperty->GetObjectPropertyValue(Property->ContainerPtrToValuePtr<UObject*>(Obj));
+			SetObjectProperty(Property->GetName(), Value);
+		}
+		else if (UStrProperty* StringProperty = Cast<UStrProperty>(Property))
+		{
+			FString Value = StringProperty->GetPropertyValue(Property->ContainerPtrToValuePtr<UObject*>(Obj));
+			SetStringProperty(Property->GetName(), Value);
+		}
+	}
+}
+
+void FScriptContextBase::FetchScriptPropertyValues(UScriptBlueprintGeneratedClass* Class, UObject* Obj)
+{
+	// @todo: optimize this
+	for (auto Property : Class->ScriptProperties)
+	{
+		if (UFloatProperty* FloatProperty = Cast<UFloatProperty>(Property))
+		{
+			float Value = GetFloatProperty(Property->GetName());
+			FloatProperty->SetFloatingPointPropertyValue(Property->ContainerPtrToValuePtr<float>(Obj), Value);
+		}
+		else if (UIntProperty* IntProperty = Cast<UIntProperty>(Property))
+		{
+			int32 Value = GetIntProperty(Property->GetName());
+			IntProperty->SetPropertyValue(Property->ContainerPtrToValuePtr<int32>(Obj), Value);
+		}
+		else if (UBoolProperty* BoolProperty = Cast<UBoolProperty>(Property))
+		{
+			bool Value = GetBoolProperty(Property->GetName());
+			BoolProperty->SetPropertyValue(Property->ContainerPtrToValuePtr<float>(Obj), Value);
+		}
+		else if (UObjectProperty* ObjectProperty = Cast<UObjectProperty>(Property))
+		{
+			UObject* Value = GetObjectProperty(Property->GetName());
+			ObjectProperty->SetObjectPropertyValue(Property->ContainerPtrToValuePtr<UObject*>(Obj), Value);
+		}
+		else if (UStrProperty* StringProperty = Cast<UStrProperty>(Property))
+		{
+			FString Value = GetStringProperty(Property->GetName());
+			StringProperty->SetPropertyValue(Property->ContainerPtrToValuePtr<FString>(Obj), Value);
+		}
+	}
 }
\ No newline at end of file
diff --git a/Engine/Plugins/Script/ScriptPlugin/Source/ScriptPlugin/Private/ScriptComponent.cpp b/Engine/Plugins/Script/ScriptPlugin/Source/ScriptPlugin/Private/ScriptComponent.cpp
--- a/Engine/Plugins/Script/ScriptPlugin/Source/ScriptPlugin/Private/ScriptComponent.cpp
+++ b/Engine/Plugins/Script/ScriptPlugin/Source/ScriptPlugin/Private/ScriptComponent.cpp
@@ -18,12 +18,7 @@ void UScriptComponent::OnRegister()
 {
 	Super::OnRegister();
 
-	UScriptBlueprintGeneratedClass* ScriptClass = NULL;
-	for (UClass* MyClass = GetClass(); MyClass && !ScriptClass; MyClass = MyClass->GetSuperClass())
-	{
-		ScriptClass = Cast<UScriptBlueprintGeneratedClass>(MyClass);
-	}
-
+	auto ScriptClass = UScriptBlueprintGeneratedClass::GetScriptGeneratedClass(GetClass());
 	if (ScriptClass && GetWorld() && GetWorld()->WorldType != EWorldType::Editor)
 	{
 		Context = ScriptClass->CreateContext();
@@ -33,7 +28,10 @@ void UScriptComponent::OnRegister()
 			if (Context->Initialize(ScriptClass->SourceCode, this))
 			{
 				bDoNotTick = !Context->CanTick();
+				// Push values set by CDO
+				Context->PushScriptPropertyValues(ScriptClass, this);
 			}
+
 			if (bDoNotTick)
 			{
 				bAutoActivate = false;
@@ -48,7 +46,10 @@ void UScriptComponent::InitializeComponent()
 	Super::InitializeComponent();
 	if (Context)
 	{
+		auto ScriptClass = UScriptBlueprintGeneratedClass::GetScriptGeneratedClass(GetClass());
+		Context->PushScriptPropertyValues(ScriptClass, this);
 		Context->BeginPlay();
+		Context->FetchScriptPropertyValues(ScriptClass, this);
 	}
 }
 
@@ -57,7 +58,10 @@ void UScriptComponent::TickComponent(float DeltaTime, enum ELevelTick TickType,
 	Super::TickComponent(DeltaTime, TickType, ThisTickFunction);
 	if (Context)
 	{
+		auto ScriptClass = UScriptBlueprintGeneratedClass::GetScriptGeneratedClass(GetClass());
+		Context->PushScriptPropertyValues(ScriptClass, this);
 		Context->Tick(DeltaTime);
+		Context->FetchScriptPropertyValues(ScriptClass, this);
 	}
 };
 
@@ -78,7 +82,10 @@ bool UScriptComponent::CallScriptFunction(FString FunctionName)
 	bool bSuccess = false;
 	if (Context)
 	{
+		auto ScriptClass = UScriptBlueprintGeneratedClass::GetScriptGeneratedClass(GetClass());
+		Context->PushScriptPropertyValues(ScriptClass, this);
 		bSuccess = Context->CallFunction(FunctionName);
+		Context->FetchScriptPropertyValues(ScriptClass, this);
 	}
 	return bSuccess;
 }
