Marking ScriptComponent as Blueprintable

diff --git a/Engine/Plugins/Script/ScriptPlugin/Source/ScriptPlugin/Classes/ScriptComponent.h b/Engine/Plugins/Script/ScriptPlugin/Source/ScriptPlugin/Classes/ScriptComponent.h
--- a/Engine/Plugins/Script/ScriptPlugin/Source/ScriptPlugin/Classes/ScriptComponent.h
+++ b/Engine/Plugins/Script/ScriptPlugin/Source/ScriptPlugin/Classes/ScriptComponent.h
@@ -7,7 +7,7 @@
 /** 
  * Script-extendable component class
  */
-UCLASS(hidecategories = (Object, ActorComponent), editinlinenew, meta = (BlueprintSpawnableComponent), ClassGroup = Script, Abstract, EarlyAccessPreview)
+UCLASS(Blueprintable, hidecategories = (Object, ActorComponent), editinlinenew, meta = (BlueprintSpawnableComponent), ClassGroup = Script, Abstract, EarlyAccessPreview)
 class SCRIPTPLUGIN_API UScriptComponent : public UActorComponent
 {
 	GENERATED_UCLASS_BODY()
