ScriptPlugin: helper project for compiling lib files for Lua (moving from Extras to the plugin folder)

diff --git a/Engine/Plugins/ScriptPlugin/Source/Lua/README.txt b/Engine/Plugins/ScriptPlugin/Source/Lua/README.txt
new file mode 100644
--- /dev/null
+++ b/Engine/Plugins/ScriptPlugin/Source/Lua/README.txt
@@ -0,0 +1,7 @@
+To enable Lua integration in UE4, follow these steps:
+
+1. Download Lua source code, unzip it and copy the entire lua-5.2.3 folder to Engine/Plugins/ScriptPlugin/Source/Lua/
+2. Open Engine/Plugins/ScriptPlugin/Source/Lua/Lua.sln in Visual Studio
+3. Build Release Win32 and Release x64 configurations.
+4. Regenerate project files with GenerateProjectFiles.bat
+5. Open UE4.sln and rebuild your game target (and the editor/engine)
\ No newline at end of file
