ScriptPlugin: ScriptContext object is now marked as public so that maps with level blueprints calling functions on scripted objects can be saved without errors.

diff --git a/Engine/Plugins/ScriptPlugin/Source/ScriptEditorPlugin/Private/ScriptBlueprintCompiler.cpp b/Engine/Plugins/ScriptPlugin/Source/ScriptEditorPlugin/Private/ScriptBlueprintCompiler.cpp
--- a/Engine/Plugins/ScriptPlugin/Source/ScriptEditorPlugin/Private/ScriptBlueprintCompiler.cpp
+++ b/Engine/Plugins/ScriptPlugin/Source/ScriptEditorPlugin/Private/ScriptBlueprintCompiler.cpp
@@ -184,7 +184,7 @@ void FScriptBlueprintCompiler::FinishCompilingClass(UClass* Class)
 	if (ContextProperty)
 	{
 		UObject* CDO = Class->GetDefaultObject();
-		UObject* ContextDefaultSubobject = StaticConstructObject(ContextProperty->PropertyClass, CDO, "ScriptContext", RF_DefaultSubObject);
+		UObject* ContextDefaultSubobject = StaticConstructObject(ContextProperty->PropertyClass, CDO, "ScriptContext", RF_DefaultSubObject|RF_Public);
 		ContextProperty->SetObjectPropertyValue(ContextProperty->ContainerPtrToValuePtr<UObject*>(CDO), ContextDefaultSubobject);
 	}
 }
