[INTEGRATE] Change 2427264 by [REDACTED_PERSON][REDACTED_HANDLE] on 2015/01/31 18:51:47

	Renaming ScriptComponent to ScriptPluginComponent

diff --git a/Engine/Config/BaseEngine.ini b/Engine/Config/BaseEngine.ini
--- a/Engine/Config/BaseEngine.ini
+++ b/Engine/Config/BaseEngine.ini
@@ -450,6 +450,9 @@ bUseStreamingPause=false
 +ActiveClassRedirects=(OldClassName="SReply",NewClassName="EventReply")
 +K2FieldRedirects=(OldFieldName="Widget.Visiblity", NewFieldName="Widget.Visibility")
 
+; Renamed ScriptComponent to ScriptPluginComponent
++ActiveClassRedirects=(OldClassName="ScriptComponent",NewClassName="ScriptPluginComponent")
+
 ;EAxis rename
 +EnumRedirects=(EnumName="EAxis",OldEnumEntry="AXIS_X",NewEnumEntry="EAxis::X")
 +EnumRedirects=(EnumName="EAxis",OldEnumEntry="AXIS_Y",NewEnumEntry="EAxis::Y")
diff --git a/Engine/Plugins/ScriptPlugin/Source/ScriptEditorPlugin/Private/ScriptBlueprintCompiler.cpp b/Engine/Plugins/ScriptPlugin/Source/ScriptEditorPlugin/Private/ScriptBlueprintCompiler.cpp
--- a/Engine/Plugins/ScriptPlugin/Source/ScriptEditorPlugin/Private/ScriptBlueprintCompiler.cpp
+++ b/Engine/Plugins/ScriptPlugin/Source/ScriptEditorPlugin/Private/ScriptBlueprintCompiler.cpp
@@ -87,13 +87,13 @@ void FScriptBlueprintCompiler::CreateClassVariablesFromBlueprint()
 
 void FScriptBlueprintCompiler::CreateScriptContextProperty()
 {
-	// The only case we don't need a script context is if the script class derives form UScriptComponent
+	// The only case we don't need a script context is if the script class derives form UScriptPluginComponent
 	UClass* ContextClass = nullptr;
 	if (Blueprint->ParentClass->IsChildOf(AActor::StaticClass()))
 	{
 		ContextClass = UScriptContextComponent::StaticClass();
 	}
-	else if (!Blueprint->ParentClass->IsChildOf(UScriptComponent::StaticClass()))
+	else if (!Blueprint->ParentClass->IsChildOf(UScriptPluginComponent::StaticClass()))
 	{
 		ContextClass = UScriptContext::StaticClass();
 	}
@@ -110,7 +110,7 @@ void FScriptBlueprintCompiler::CreateFunctionList()
 {
 	Super::CreateFunctionList();
 
-	if (!Blueprint->ParentClass->IsChildOf(UScriptComponent::StaticClass()))
+	if (!Blueprint->ParentClass->IsChildOf(UScriptPluginComponent::StaticClass()))
 	{
 		for (auto& Field : ScriptDefinedFields)
 		{
@@ -173,7 +173,7 @@ void FScriptBlueprintCompiler::FinishCompilingClass(UClass* Class)
 	ScriptClass->ByteCode = Blueprint->ByteCode;
 
 	// Allow Blueprint Components to be used in Blueprints
-	if (Blueprint->ParentClass->IsChildOf(UScriptComponent::StaticClass()) && Class != Blueprint->SkeletonGeneratedClass)
+	if (Blueprint->ParentClass->IsChildOf(UScriptPluginComponent::StaticClass()) && Class != Blueprint->SkeletonGeneratedClass)
 	{
 		Class->SetMetaData(TEXT("BlueprintSpawnableComponent"), TEXT("true"));
 	}
diff --git a/Engine/Plugins/ScriptPlugin/Source/ScriptEditorPlugin/Private/ScriptEditorPluginPrivatePCH.h b/Engine/Plugins/ScriptPlugin/Source/ScriptEditorPlugin/Private/ScriptEditorPluginPrivatePCH.h
--- a/Engine/Plugins/ScriptPlugin/Source/ScriptEditorPlugin/Private/ScriptEditorPluginPrivatePCH.h
+++ b/Engine/Plugins/ScriptPlugin/Source/ScriptEditorPlugin/Private/ScriptEditorPluginPrivatePCH.h
@@ -10,7 +10,7 @@
 #include "ScriptFactory.h"
 #include "ReimportScriptFactory.h"
 #include "AssetTypeActions_Base.h"
-#include "ScriptComponent.h"
+#include "ScriptPluginComponent.h"
 #include "ClassViewerModule.h"
 #include "ClassViewerFilter.h"
 #include "Kismet2/KismetEditorUtilities.h"
diff --git a/Engine/Plugins/ScriptPlugin/Source/ScriptPlugin/Classes/ScriptComponent.h b/Engine/Plugins/ScriptPlugin/Source/ScriptPlugin/Classes/ScriptComponent.h
deleted file mode 100644
diff --git a/Engine/Plugins/ScriptPlugin/Source/ScriptPlugin/Classes/ScriptPluginComponent.h b/Engine/Plugins/ScriptPlugin/Source/ScriptPlugin/Classes/ScriptPluginComponent.h
new file mode 100644
--- /dev/null
+++ b/Engine/Plugins/ScriptPlugin/Source/ScriptPlugin/Classes/ScriptPluginComponent.h
@@ -0,0 +1,36 @@
+// Copyright 1998-2015 Epic Games, Inc. All Rights Reserved.
+
+#pragma once
+#include "ScriptBlueprintGeneratedClass.h"
+#include "ScriptPluginComponent.generated.h"
+
+/** 
+ * Script-extendable component class
+ */
+UCLASS(Blueprintable, hidecategories = (Object, ActorComponent), editinlinenew, meta = (BlueprintSpawnableComponent), ClassGroup = Script, Abstract, EarlyAccessPreview)
+class SCRIPTPLUGIN_API UScriptPluginComponent : public UActorComponent
+{
+	GENERATED_UCLASS_BODY()
+
+public:
+
+	/**
+	* Calls a script-defined function (no arguments)
+	* @param FunctionName Name of the function to call
+	*/
+	UFUNCTION(BlueprintCallable, Category = "Script|Functions")
+	virtual bool CallScriptFunction(FString FunctionName);
+
+	// Begin UActorComponent interface.
+	virtual void OnRegister() override;
+	virtual void InitializeComponent() override;
+	virtual void TickComponent(float DeltaTime, enum ELevelTick TickType, FActorComponentTickFunction *ThisTickFunction) override;
+	virtual void OnUnregister() override;
+	// Begin UActorComponent interface.
+
+protected:
+
+	/** Script context (code) */
+	FScriptContextBase* Context;
+};
+
diff --git a/Engine/Plugins/ScriptPlugin/Source/ScriptPlugin/Private/ScriptComponent.cpp b/Engine/Plugins/ScriptPlugin/Source/ScriptPlugin/Private/ScriptComponent.cpp
deleted file mode 100644
diff --git a/Engine/Plugins/ScriptPlugin/Source/ScriptPlugin/Private/ScriptPluginComponent.cpp b/Engine/Plugins/ScriptPlugin/Source/ScriptPlugin/Private/ScriptPluginComponent.cpp
new file mode 100644
--- /dev/null
+++ b/Engine/Plugins/ScriptPlugin/Source/ScriptPlugin/Private/ScriptPluginComponent.cpp
@@ -0,0 +1,80 @@
+// Copyright 1998-2015 Epic Games, Inc. All Rights Reserved. 
+#include "ScriptPluginPrivatePCH.h"
+
+//////////////////////////////////////////////////////////////////////////
+
+UScriptPluginComponent::UScriptPluginComponent(const FObjectInitializer& ObjectInitializer)
+	: Super(ObjectInitializer)
+{
+	PrimaryComponentTick.bCanEverTick = true;
+	bTickInEditor = false;
+	bAutoActivate = true;
+	bWantsInitializeComponent = true;
+
+	Context = NULL;
+}
+
+void UScriptPluginComponent::OnRegister()
+{
+	Super::OnRegister();
+
+	auto ScriptClass = UScriptBlueprintGeneratedClass::GetScriptGeneratedClass(GetClass());
+	if (ScriptClass && GetWorld() && GetWorld()->WorldType != EWorldType::Editor)
+	{
+		Context = FScriptContextBase::CreateContext(ScriptClass->SourceCode, ScriptClass, this);
+		if (!Context || !Context->CanTick())
+		{
+			bAutoActivate = false;
+			PrimaryComponentTick.bCanEverTick = false;
+		}
+	}
+}
+
+void UScriptPluginComponent::InitializeComponent()
+{
+	Super::InitializeComponent();
+	if (Context)
+	{
+		auto ScriptClass = UScriptBlueprintGeneratedClass::GetScriptGeneratedClass(GetClass());
+		Context->PushScriptPropertyValues(ScriptClass, this);
+		Context->BeginPlay();
+		Context->FetchScriptPropertyValues(ScriptClass, this);
+	}
+}
+
+void UScriptPluginComponent::TickComponent(float DeltaTime, enum ELevelTick TickType, FActorComponentTickFunction *ThisTickFunction)
+{
+	Super::TickComponent(DeltaTime, TickType, ThisTickFunction);
+	if (Context)
+	{
+		auto ScriptClass = UScriptBlueprintGeneratedClass::GetScriptGeneratedClass(GetClass());
+		Context->PushScriptPropertyValues(ScriptClass, this);
+		Context->Tick(DeltaTime);
+		Context->FetchScriptPropertyValues(ScriptClass, this);
+	}
+};
+
+void UScriptPluginComponent::OnUnregister()
+{
+	if (Context)
+	{
+		Context->Destroy();
+		delete Context;
+		Context = NULL;
+	}
+
+	Super::OnUnregister();
+}
+
+bool UScriptPluginComponent::CallScriptFunction(FString FunctionName)
+{
+	bool bSuccess = false;
+	if (Context)
+	{
+		auto ScriptClass = UScriptBlueprintGeneratedClass::GetScriptGeneratedClass(GetClass());
+		Context->PushScriptPropertyValues(ScriptClass, this);
+		bSuccess = Context->CallFunction(FunctionName);
+		Context->FetchScriptPropertyValues(ScriptClass, this);
+	}
+	return bSuccess;
+}
diff --git a/Engine/Plugins/ScriptPlugin/Source/ScriptPlugin/Private/ScriptPluginPrivatePCH.h b/Engine/Plugins/ScriptPlugin/Source/ScriptPlugin/Private/ScriptPluginPrivatePCH.h
--- a/Engine/Plugins/ScriptPlugin/Source/ScriptPlugin/Private/ScriptPluginPrivatePCH.h
+++ b/Engine/Plugins/ScriptPlugin/Source/ScriptPlugin/Private/ScriptPluginPrivatePCH.h
@@ -5,7 +5,7 @@
 #include "CoreUObject.h"
 #include "ModuleManager.h"
 #include "Engine.h"
-#include "ScriptComponent.h"
+#include "ScriptPluginComponent.h"
 #include "ScriptContext.h"
 #include "ScriptContextComponent.h"
 #include "ScriptTestActor.h"
