ScriptPlugin will no longer import txt files

diff --git a/Engine/Plugins/ScriptPlugin/Source/ScriptEditorPlugin/Private/ScriptFactory.cpp b/Engine/Plugins/ScriptPlugin/Source/ScriptEditorPlugin/Private/ScriptFactory.cpp
--- a/Engine/Plugins/ScriptPlugin/Source/ScriptEditorPlugin/Private/ScriptFactory.cpp
+++ b/Engine/Plugins/ScriptPlugin/Source/ScriptEditorPlugin/Private/ScriptFactory.cpp
@@ -7,7 +7,6 @@ UScriptFactory::UScriptFactory(const FObjectInitializer& ObjectInitializer)
 	SupportedClass = UScriptBlueprint::StaticClass();
 	ParentClass = AActor::StaticClass();
 
-	Formats.Add(TEXT("txt;Script"));
 	Formats.Add(TEXT("lua;Script"));
 
 	bCreateNew = false;
