Adding script contents validation when importing scripts for the script plugin. Added support for custom file extensions for script plugin.

diff --git a/Engine/Plugins/ScriptPlugin/Source/ScriptEditorPlugin/Private/ScriptFactory.cpp b/Engine/Plugins/ScriptPlugin/Source/ScriptEditorPlugin/Private/ScriptFactory.cpp
--- a/Engine/Plugins/ScriptPlugin/Source/ScriptEditorPlugin/Private/ScriptFactory.cpp
+++ b/Engine/Plugins/ScriptPlugin/Source/ScriptEditorPlugin/Private/ScriptFactory.cpp
@@ -1,5 +1,6 @@
 // Copyright 1998-2015 Epic Games, Inc. All Rights Reserved. 
 #include "ScriptEditorPluginPrivatePCH.h"
+#include "ScriptBlueprintGeneratedClass.h"
 
 UScriptFactory::UScriptFactory(const FObjectInitializer& ObjectInitializer)
 	: Super( ObjectInitializer )
@@ -7,7 +8,7 @@ UScriptFactory::UScriptFactory(const FObjectInitializer& ObjectInitializer)
 	SupportedClass = UScriptBlueprint::StaticClass();
 	ParentClass = AActor::StaticClass();
 
-	Formats.Add(TEXT("lua;Script"));
+	FScriptContextBase::GetSupportedScriptFileFormats(Formats);
 
 	bCreateNew = false;
 	bEditorImport = true;
@@ -52,15 +53,24 @@ UObject* UScriptFactory::FactoryCreateText(UClass* InClass, UObject* InParent, F
 {
 	GEditor->SelectNone(true, true, false);
 
-	UScriptBlueprint* NewBlueprint = CastChecked<UScriptBlueprint>(FKismetEditorUtilities::CreateBlueprint(ParentClass, InParent, InName, BPTYPE_Normal, UScriptBlueprint::StaticClass(), UScriptBlueprintGeneratedClass::StaticClass(), "UScriptFactory"));
-	NewBlueprint->SourceFilePath = FReimportManager::SanitizeImportFilename(CurrentFilename, NewBlueprint);
-	NewBlueprint->SourceCode = Buffer;
-	NewBlueprint->SourceFileTimestamp = IFileManager::Get().GetTimeStamp(*NewBlueprint->SourceFilePath).ToString();
+	UScriptBlueprint* NewBlueprint = nullptr;
+	TAutoPtr<FScriptContextBase> ScriptContext(FScriptContextBase::CreateContext(Buffer, nullptr, nullptr));
+	if (ScriptContext.IsValid())
+	{
+		UScriptBlueprint* NewBlueprint = CastChecked<UScriptBlueprint>(FKismetEditorUtilities::CreateBlueprint(ParentClass, InParent, InName, BPTYPE_Normal, UScriptBlueprint::StaticClass(), UScriptBlueprintGeneratedClass::StaticClass(), "UScriptFactory"));
+		NewBlueprint->SourceFilePath = FReimportManager::SanitizeImportFilename(CurrentFilename, NewBlueprint);
+		NewBlueprint->SourceCode = Buffer;
+		NewBlueprint->SourceFileTimestamp = IFileManager::Get().GetTimeStamp(*NewBlueprint->SourceFilePath).ToString();
 
-	// Need to make sure we compile with the new source code
-	FKismetEditorUtilities::CompileBlueprint(NewBlueprint);
+		// Need to make sure we compile with the new source code
+		FKismetEditorUtilities::CompileBlueprint(NewBlueprint);
 
-	FEditorDelegates::OnAssetPostImport.Broadcast(this, NewBlueprint);
+		FEditorDelegates::OnAssetPostImport.Broadcast(this, NewBlueprint);
+	}
+	else
+	{
+		UE_LOG(LogScriptEditorPlugin, Warning, TEXT("Failed to import %s: could not compile script."), *CurrentFilename);
+	}
 
 	return NewBlueprint;
 }
diff --git a/Engine/Plugins/ScriptPlugin/Source/ScriptPlugin/Classes/ScriptBlueprintGeneratedClass.h b/Engine/Plugins/ScriptPlugin/Source/ScriptPlugin/Classes/ScriptBlueprintGeneratedClass.h
--- a/Engine/Plugins/ScriptPlugin/Source/ScriptPlugin/Classes/ScriptBlueprintGeneratedClass.h
+++ b/Engine/Plugins/ScriptPlugin/Source/ScriptPlugin/Classes/ScriptBlueprintGeneratedClass.h
@@ -40,6 +40,11 @@ public:
 	*/
 	static FScriptContextBase* CreateContext(const FString& SourceCode, UScriptBlueprintGeneratedClass* Class, UObject* Owner);
 
+	/**
+	* Gets supported script file formats (UFactory formats)
+	*/
+	static void GetSupportedScriptFileFormats(TArray<FString>& OutFormats);
+
 	/**
 	 * Initializes script context given script code
 	 * @param Code Script code
diff --git a/Engine/Plugins/ScriptPlugin/Source/ScriptPlugin/Private/ScriptBlueprintGeneratedClass.cpp b/Engine/Plugins/ScriptPlugin/Source/ScriptPlugin/Private/ScriptBlueprintGeneratedClass.cpp
--- a/Engine/Plugins/ScriptPlugin/Source/ScriptPlugin/Private/ScriptBlueprintGeneratedClass.cpp
+++ b/Engine/Plugins/ScriptPlugin/Source/ScriptPlugin/Private/ScriptBlueprintGeneratedClass.cpp
@@ -136,4 +136,11 @@ FScriptContextBase* FScriptContextBase::CreateContext(const FString& SourceCode,
 		}
 	}
 	return NewContext;
+}
+
+void FScriptContextBase::GetSupportedScriptFileFormats(TArray<FString>& OutFormats)
+{
+#if WITH_LUA
+	OutFormats.Add(TEXT("lua;Script"));
+#endif
 }
\ No newline at end of file
