Stop excluding ScriptPlugin and ScriptGeneratorPlugin from Rocket builds.

diff --git a/Engine/Plugins/ScriptGeneratorPlugin/Source/ScriptGeneratorPlugin/ScriptGeneratorPlugin.Build.cs b/Engine/Plugins/ScriptGeneratorPlugin/Source/ScriptGeneratorPlugin/ScriptGeneratorPlugin.Build.cs
--- a/Engine/Plugins/ScriptGeneratorPlugin/Source/ScriptGeneratorPlugin/ScriptGeneratorPlugin.Build.cs
+++ b/Engine/Plugins/ScriptGeneratorPlugin/Source/ScriptGeneratorPlugin/ScriptGeneratorPlugin.Build.cs
@@ -45,21 +45,18 @@ namespace UnrealBuildTool.Rules
 				}
 				);
 
-			if (!UnrealBuildTool.BuildingRocket())
+			// This checks only for UHT target platform, not the target platform of the game we're building so it's important
+			// to make sure Lua is compiled for all supported platforms
+			var LuaLibDirectory = Path.Combine("..", "Plugins", "ScriptPlugin", "Source", "Lua", "Lib", Target.Platform.ToString(), "Release");
+			var LuaLibPath = Path.Combine(LuaLibDirectory, "Lua.lib");
+			if (File.Exists(LuaLibPath))
 			{
-				// This checks only for UHT target platform, not the target platform of the game we're building so it's important
-				// to make sure Lua is compiled for all supported platforms
-				var LuaLibDirectory = Path.Combine("..", "Plugins", "ScriptPlugin", "Source", "Lua", "Lib", Target.Platform.ToString(), "Release");
-				var LuaLibPath = Path.Combine(LuaLibDirectory, "Lua.lib");
-				if (File.Exists(LuaLibPath))
-				{
-					Log.TraceVerbose("ScriptGenerator LUA Integration enabled");
-					Definitions.Add("WITH_LUA=1");
-				}
-				else
-				{
-					Log.TraceVerbose("ScriptGenerator LUA Integration NOT enabled");
-				}
+				Log.TraceVerbose("ScriptGenerator LUA Integration enabled");
+				Definitions.Add("WITH_LUA=1");
+			}
+			else
+			{
+				Log.TraceVerbose("ScriptGenerator LUA Integration NOT enabled");
 			}
 		}
 	}
diff --git a/Engine/Plugins/ScriptPlugin/Source/ScriptPlugin/ScriptPlugin.Build.cs b/Engine/Plugins/ScriptPlugin/Source/ScriptPlugin/ScriptPlugin.Build.cs
--- a/Engine/Plugins/ScriptPlugin/Source/ScriptPlugin/ScriptPlugin.Build.cs
+++ b/Engine/Plugins/ScriptPlugin/Source/ScriptPlugin/ScriptPlugin.Build.cs
@@ -53,29 +53,26 @@ namespace UnrealBuildTool.Rules
 				}
 				);
 
-			if (!UnrealBuildTool.BuildingRocket())
-			{
-				var LuaPath = Path.Combine("..", "Plugins", "ScriptPlugin", "Source", "Lua");				
-				var LuaLibDirectory = Path.Combine(LuaPath, "Lib", Target.Platform.ToString(), "Release");
-				var LuaLibPath = Path.Combine(LuaLibDirectory, "Lua.lib");
-				if (File.Exists(LuaLibPath))
-				{					
-					Definitions.Add("WITH_LUA=1");
+			var LuaPath = Path.Combine("..", "Plugins", "ScriptPlugin", "Source", "Lua");				
+			var LuaLibDirectory = Path.Combine(LuaPath, "Lib", Target.Platform.ToString(), "Release");
+			var LuaLibPath = Path.Combine(LuaLibDirectory, "Lua.lib");
+			if (File.Exists(LuaLibPath))
+			{					
+				Definitions.Add("WITH_LUA=1");
 
-					// Path to Lua include files
-					var IncludePath = Path.GetFullPath(Path.Combine(LuaPath, "Include"));
-					PrivateIncludePaths.Add(IncludePath);
+				// Path to Lua include files
+				var IncludePath = Path.GetFullPath(Path.Combine(LuaPath, "Include"));
+				PrivateIncludePaths.Add(IncludePath);
 
-					// Lib file
-					PublicLibraryPaths.Add(LuaLibDirectory);
-					PublicAdditionalLibraries.Add(LuaLibPath);
+				// Lib file
+				PublicLibraryPaths.Add(LuaLibDirectory);
+				PublicAdditionalLibraries.Add(LuaLibPath);
 
-					Log.TraceVerbose("LUA Integration enabled: {0}", IncludePath);
-				}
-				else
-				{
-					Log.TraceVerbose("LUA Integration NOT enabled");
-				}
+				Log.TraceVerbose("LUA Integration enabled: {0}", IncludePath);
+			}
+			else
+			{
+				Log.TraceVerbose("LUA Integration NOT enabled");
 			}
 		}
 	}
