Fixing CIS - temporarily disabling script plugin integration. Script pluging should be a game-only plugin.

diff --git a/Engine/Plugins/ScriptPlugin/Source/ScriptPlugin/Private/ScriptPlugin.cpp b/Engine/Plugins/ScriptPlugin/Source/ScriptPlugin/Private/ScriptPlugin.cpp
--- a/Engine/Plugins/ScriptPlugin/Source/ScriptPlugin/Private/ScriptPlugin.cpp
+++ b/Engine/Plugins/ScriptPlugin/Source/ScriptPlugin/Private/ScriptPlugin.cpp
@@ -16,7 +16,9 @@ UProperty* FindScriptPropertyHelper(UClass* Class, FName PropertyName)
 	return NULL;
 }
 
+#if 0
 #include "GeneratedScriptLibraries.inl"
+#endif 
 
 DEFINE_LOG_CATEGORY(LogScriptPlugin);
 
