Fixed assertion when reimporting assets for script plugin components

diff --git a/Engine/Plugins/ScriptPlugin/Source/ScriptEditorPlugin/Private/ScriptFactory.cpp b/Engine/Plugins/ScriptPlugin/Source/ScriptEditorPlugin/Private/ScriptFactory.cpp
--- a/Engine/Plugins/ScriptPlugin/Source/ScriptEditorPlugin/Private/ScriptFactory.cpp
+++ b/Engine/Plugins/ScriptPlugin/Source/ScriptEditorPlugin/Private/ScriptFactory.cpp
@@ -57,7 +57,15 @@ UObject* UScriptFactory::FactoryCreateText(UClass* InClass, UObject* InParent, F
 	TAutoPtr<FScriptContextBase> ScriptContext(FScriptContextBase::CreateContext(Buffer, nullptr, nullptr));
 	if (ScriptContext.IsValid())
 	{
-		UScriptBlueprint* NewBlueprint = CastChecked<UScriptBlueprint>(FKismetEditorUtilities::CreateBlueprint(ParentClass, InParent, InName, BPTYPE_Normal, UScriptBlueprint::StaticClass(), UScriptBlueprintGeneratedClass::StaticClass(), "UScriptFactory"));
+	    UScriptBlueprint* NewBlueprint = Cast<UScriptBlueprint>( FindObject<UBlueprint>( InParent, *InName.ToString() ) );
+	    if( NewBlueprint != nullptr )
+	    {
+		    NewBlueprint->Modify();
+	    }
+	    else
+	    {
+		    NewBlueprint = CastChecked<UScriptBlueprint>(FKismetEditorUtilities::CreateBlueprint(ParentClass, InParent, InName, BPTYPE_Normal, UScriptBlueprint::StaticClass(), UScriptBlueprintGeneratedClass::StaticClass(), "UScriptFactory"));
+		}
 		NewBlueprint->SourceFilePath = FReimportManager::SanitizeImportFilename(CurrentFilename, NewBlueprint);
 		NewBlueprint->SourceCode = Buffer;
 		NewBlueprint->SourceFileTimestamp = IFileManager::Get().GetTimeStamp(*NewBlueprint->SourceFilePath).ToString();
