Updated the example Lua integration to use Lua 5.3.0 (latest version)

diff --git a/Engine/Plugins/ScriptPlugin/Source/Lua/README.txt b/Engine/Plugins/ScriptPlugin/Source/Lua/README.txt
--- a/Engine/Plugins/ScriptPlugin/Source/Lua/README.txt
+++ b/Engine/Plugins/ScriptPlugin/Source/Lua/README.txt
@@ -1,6 +1,6 @@
 To enable Lua integration in UE4, follow these steps:
 
-1. Download Lua source code, unzip it and copy the entire lua-5.2.3 folder to Engine/Plugins/ScriptPlugin/Source/Lua/
+1. Download Lua source code, unzip it and copy the entire lua-5.3.0 folder to Engine/Plugins/ScriptPlugin/Source/Lua/
 2. Open Engine/Plugins/ScriptPlugin/Source/Lua/Lua.sln in Visual Studio
 3. Build Release Win32 and Release x64 configurations.
 4. Regenerate project files with GenerateProjectFiles.bat
diff --git a/Engine/Plugins/ScriptPlugin/Source/ScriptPlugin/Private/LuaIntegration.h b/Engine/Plugins/ScriptPlugin/Source/ScriptPlugin/Private/LuaIntegration.h
--- a/Engine/Plugins/ScriptPlugin/Source/ScriptPlugin/Private/LuaIntegration.h
+++ b/Engine/Plugins/ScriptPlugin/Source/ScriptPlugin/Private/LuaIntegration.h
@@ -5,6 +5,7 @@
 
 extern "C"
 {
+#define LUA_COMPAT_APIINTCASTS 1
 #include "lua.h"
 #include "lauxlib.h"
 #include "lualib.h"
