Fix ScriptFactory not returning created blueprint

diff --git a/Engine/Plugins/ScriptPlugin/Source/ScriptEditorPlugin/Private/ScriptFactory.cpp b/Engine/Plugins/ScriptPlugin/Source/ScriptEditorPlugin/Private/ScriptFactory.cpp
--- a/Engine/Plugins/ScriptPlugin/Source/ScriptEditorPlugin/Private/ScriptFactory.cpp
+++ b/Engine/Plugins/ScriptPlugin/Source/ScriptEditorPlugin/Private/ScriptFactory.cpp
@@ -57,7 +57,7 @@ UObject* UScriptFactory::FactoryCreateText(UClass* InClass, UObject* InParent, F
 	TAutoPtr<FScriptContextBase> ScriptContext(FScriptContextBase::CreateContext(Buffer, nullptr, nullptr));
 	if (ScriptContext.IsValid())
 	{
-	    UScriptBlueprint* NewBlueprint = Cast<UScriptBlueprint>( FindObject<UBlueprint>( InParent, *InName.ToString() ) );
+	    NewBlueprint = Cast<UScriptBlueprint>( FindObject<UBlueprint>( InParent, *InName.ToString() ) );
 	    if( NewBlueprint != nullptr )
 	    {
 		    NewBlueprint->Modify();
