Removed #if 0'd out include of GeneratedScriptLibraries.inl because it was being regenerated between builds of UE4Editor and QAGameEditor, causing the build system to balk at the build product files being tampered with.

diff --git a/Engine/Plugins/ScriptPlugin/Source/ScriptPlugin/Private/ScriptPlugin.cpp b/Engine/Plugins/ScriptPlugin/Source/ScriptPlugin/Private/ScriptPlugin.cpp
--- a/Engine/Plugins/ScriptPlugin/Source/ScriptPlugin/Private/ScriptPlugin.cpp
+++ b/Engine/Plugins/ScriptPlugin/Source/ScriptPlugin/Private/ScriptPlugin.cpp
@@ -16,10 +16,6 @@ UProperty* FindScriptPropertyHelper(UClass* Class, FName PropertyName)
 	return NULL;
 }
 
-#if 0
-#include "GeneratedScriptLibraries.inl"
-#endif 
-
 DEFINE_LOG_CATEGORY(LogScriptPlugin);
 
 class FScriptPlugin : public IScriptPlugin
