Fix compile errors in ScriptGeneratorPlugin.

#rnx

diff --git a/Engine/Plugins/ScriptPlugin/Source/ScriptGeneratorPlugin/Private/GenericScriptCodeGenerator.cpp b/Engine/Plugins/ScriptPlugin/Source/ScriptGeneratorPlugin/Private/GenericScriptCodeGenerator.cpp
--- a/Engine/Plugins/ScriptPlugin/Source/ScriptGeneratorPlugin/Private/GenericScriptCodeGenerator.cpp
+++ b/Engine/Plugins/ScriptPlugin/Source/ScriptGeneratorPlugin/Private/GenericScriptCodeGenerator.cpp
@@ -3,6 +3,8 @@
 #include "GenericScriptCodeGenerator.h"
 #include "UObject/UnrealType.h"
 #include "ScriptGeneratorLog.h"
+#include "Misc/Paths.h"
+#include "UObject/Class.h"
 
 FGenericScriptCodeGenerator::FGenericScriptCodeGenerator(const FString& RootLocalPath, const FString& RootBuildPath, const FString& OutputDirectory, const FString& InIncludeBase)
 : FScriptCodeGeneratorBase(RootLocalPath, RootBuildPath, OutputDirectory, InIncludeBase)
diff --git a/Engine/Plugins/ScriptPlugin/Source/ScriptGeneratorPlugin/Private/LuaScriptCodeGenerator.cpp b/Engine/Plugins/ScriptPlugin/Source/ScriptGeneratorPlugin/Private/LuaScriptCodeGenerator.cpp
--- a/Engine/Plugins/ScriptPlugin/Source/ScriptGeneratorPlugin/Private/LuaScriptCodeGenerator.cpp
+++ b/Engine/Plugins/ScriptPlugin/Source/ScriptGeneratorPlugin/Private/LuaScriptCodeGenerator.cpp
@@ -3,6 +3,7 @@
 #include "LuaScriptCodeGenerator.h"
 #include "ScriptGeneratorLog.h"
 #include "UObject/ErrorException.h"
+#include "UObject/Class.h"
 
 // Supported structs
 static FName Name_Vector2D("Vector2D");
diff --git a/Engine/Plugins/ScriptPlugin/Source/ScriptGeneratorPlugin/Private/ScriptCodeGeneratorBase.cpp b/Engine/Plugins/ScriptPlugin/Source/ScriptGeneratorPlugin/Private/ScriptCodeGeneratorBase.cpp
--- a/Engine/Plugins/ScriptPlugin/Source/ScriptGeneratorPlugin/Private/ScriptCodeGeneratorBase.cpp
+++ b/Engine/Plugins/ScriptPlugin/Source/ScriptGeneratorPlugin/Private/ScriptCodeGeneratorBase.cpp
@@ -2,6 +2,8 @@
 
 #include "ScriptCodeGeneratorBase.h"
 #include "Misc/FileHelper.h"
+#include "Misc/Paths.h"
+#include "UObject/Class.h"
 #include "ScriptGeneratorLog.h"
 
 FScriptCodeGeneratorBase::FScriptCodeGeneratorBase(const FString& InRootLocalPath, const FString& InRootBuildPath, const FString& InOutputDirectory, const FString& InIncludeBase)
diff --git a/Engine/Plugins/ScriptPlugin/Source/ScriptGeneratorPlugin/Private/ScriptGeneratorPlugin.cpp b/Engine/Plugins/ScriptPlugin/Source/ScriptGeneratorPlugin/Private/ScriptGeneratorPlugin.cpp
--- a/Engine/Plugins/ScriptPlugin/Source/ScriptGeneratorPlugin/Private/ScriptGeneratorPlugin.cpp
+++ b/Engine/Plugins/ScriptPlugin/Source/ScriptGeneratorPlugin/Private/ScriptGeneratorPlugin.cpp
@@ -8,6 +8,7 @@
 #include "Interfaces/IProjectManager.h"
 #include "Features/IModularFeatures.h"
 #include "Templates/UniquePtr.h"
+#include "Misc/ConfigCacheIni.h"
 #include "ProjectDescriptor.h"
 
 DEFINE_LOG_CATEGORY(LogScriptGenerator);
diff --git a/Engine/Plugins/ScriptPlugin/Source/ScriptGeneratorPlugin/Public/IScriptGeneratorPlugin.h b/Engine/Plugins/ScriptPlugin/Source/ScriptGeneratorPlugin/Public/IScriptGeneratorPlugin.h
--- a/Engine/Plugins/ScriptPlugin/Source/ScriptGeneratorPlugin/Public/IScriptGeneratorPlugin.h
+++ b/Engine/Plugins/ScriptPlugin/Source/ScriptGeneratorPlugin/Public/IScriptGeneratorPlugin.h
@@ -2,6 +2,8 @@
 
 #pragma once
 
+#include "CoreTypes.h"
+#include "Modules/ModuleManager.h"
 #include "IScriptGeneratorPluginInterface.h"
 
 
diff --git a/Engine/Plugins/ScriptPlugin/Source/ScriptGeneratorPlugin/ScriptGeneratorPlugin.Build.cs b/Engine/Plugins/ScriptPlugin/Source/ScriptGeneratorPlugin/ScriptGeneratorPlugin.Build.cs
--- a/Engine/Plugins/ScriptPlugin/Source/ScriptGeneratorPlugin/ScriptGeneratorPlugin.Build.cs
+++ b/Engine/Plugins/ScriptPlugin/Source/ScriptGeneratorPlugin/ScriptGeneratorPlugin.Build.cs
@@ -17,7 +17,6 @@ namespace UnrealBuildTool.Rules
 
 			PrivateIncludePaths.AddRange(
 				new string[] {
-					"Developer/ScriptGeneratorPlugin/Private",
 					// ... add other private include paths required here ...
 				}
 				);
