Fixed string usage issues found by UE_NODISCARD

#rnx

diff --git a/Engine/Plugins/ScriptPlugin/Source/ScriptGeneratorPlugin/Private/ScriptCodeGeneratorBase.cpp b/Engine/Plugins/ScriptPlugin/Source/ScriptGeneratorPlugin/Private/ScriptCodeGeneratorBase.cpp
--- a/Engine/Plugins/ScriptPlugin/Source/ScriptGeneratorPlugin/Private/ScriptCodeGeneratorBase.cpp
+++ b/Engine/Plugins/ScriptPlugin/Source/ScriptGeneratorPlugin/Private/ScriptCodeGeneratorBase.cpp
@@ -90,7 +90,7 @@ FString FScriptCodeGeneratorBase::GetPropertyTypeCPP(FProperty* Property, uint32
 		int32 FirstSpaceIndex = PropertyType.Find(TEXT(" "), ESearchCase::CaseSensitive);
 		PropertyType = TEXT("TEnumAsByte<") + PropertyType.Mid(FirstSpaceIndex + 1);
 	}
-	else if (PropertyType.StartsWith(TSubclassOfDecl), ESearchCase::CaseSensitive)
+	else if (PropertyType.StartsWith(TSubclassOfDecl, ESearchCase::CaseSensitive))
 	{
 		int32 FirstSpaceIndex = PropertyType.Find(TEXT(" "), ESearchCase::CaseSensitive);
 		PropertyType = TEXT("TSubclassOf<") + PropertyType.Mid(FirstSpaceIndex + 1);
