Fixed some minor issues with the sample script pliugin code for UHT

#rnx

diff --git a/Engine/Plugins/ScriptPlugin/Source/ScriptGeneratorPlugin/Private/GenericScriptCodeGenerator.cpp b/Engine/Plugins/ScriptPlugin/Source/ScriptGeneratorPlugin/Private/GenericScriptCodeGenerator.cpp
--- a/Engine/Plugins/ScriptPlugin/Source/ScriptGeneratorPlugin/Private/GenericScriptCodeGenerator.cpp
+++ b/Engine/Plugins/ScriptPlugin/Source/ScriptGeneratorPlugin/Private/GenericScriptCodeGenerator.cpp
@@ -54,13 +54,10 @@ FString FGenericScriptCodeGenerator::ExportFunction(const FString& ClassNameCPP,
 	FProperty* ReturnValue = NULL;
 	UClass* FuncSuper = NULL;
 	
-	if (Function->GetOwnerClass() != Class)
+	// Find the base class where this function was defined
+	if (Function->GetOwnerClass() != Class && CanExportClass(Function->GetOwnerClass()))
 	{
-		// Find the base definition of the function
-		if (ExportedClasses.Contains(Function->GetOwnerClass()->GetFName()))
-		{
-			FuncSuper = Function->GetOwnerClass();
-		}
+		FuncSuper = Function->GetOwnerClass();
 	}
 
 	FString FunctionBody;
@@ -98,13 +95,10 @@ FString FGenericScriptCodeGenerator::ExportProperty(const FString& ClassNameCPP,
 	FProperty* ReturnValue = NULL;
 	UClass* PropertySuper = NULL;
 
-	if (Property->GetOwnerClass() != Class)
+	// Find the base class where this property was defined
+	if (Property->GetOwnerClass() != Class && CanExportClass(Property->GetOwnerClass()))
 	{
-		// Find the base class where this property was defined
-		if (ExportedClasses.Contains(Property->GetOwnerClass()->GetFName()))
-		{
-			PropertySuper = Property->GetOwnerClass();
-		}
+		PropertySuper = Property->GetOwnerClass();
 	}
 
 	// Getter	
@@ -216,7 +210,7 @@ bool FGenericScriptCodeGenerator::CanExportClass(UClass* Class)
 
 void FGenericScriptCodeGenerator::ExportClass(UClass* Class, const FString& SourceHeaderFilename, const FString& GeneratedHeaderFilename, bool bHasChanged)
 {
-	if (!CanExportClass(Class))
+	if (!CanExportClass(Class) || ExportedClasses.Contains(Class->GetFName()))
 	{
 		return;
 	}
@@ -271,7 +265,15 @@ void FGenericScriptCodeGenerator::GlueAllGeneratedFiles()
 	FString LibGlue;
 
 	// Include all source header files
+	TArray<FString> Headers;
+	Headers.Reserve(AllSourceClassHeaders.Num());
 	for (auto& HeaderFilename : AllSourceClassHeaders)
+	{
+		Headers.Add(HeaderFilename);
+	}
+
+	Headers.Sort();
+	for (auto& HeaderFilename : Headers)
 	{
 		// Re-base to make sure we're including the right files on a remote machine
 		FString NewFilename(RebaseToBuildPath(HeaderFilename));
@@ -279,6 +281,7 @@ void FGenericScriptCodeGenerator::GlueAllGeneratedFiles()
 	}
 
 	// Include all script glue headers
+	AllScriptHeaders.Sort();
 	for (auto& HeaderFilename : AllScriptHeaders)
 	{
 		// Re-base to make sure we're including the right files on a remote machine
diff --git a/Engine/Plugins/ScriptPlugin/Source/ScriptGeneratorPlugin/Private/GenericScriptCodeGenerator.h b/Engine/Plugins/ScriptPlugin/Source/ScriptGeneratorPlugin/Private/GenericScriptCodeGenerator.h
--- a/Engine/Plugins/ScriptPlugin/Source/ScriptGeneratorPlugin/Private/GenericScriptCodeGenerator.h
+++ b/Engine/Plugins/ScriptPlugin/Source/ScriptGeneratorPlugin/Private/GenericScriptCodeGenerator.h
@@ -15,7 +15,7 @@ protected:
 	/** All generated script header filenames */
 	TArray<FString> AllScriptHeaders;
 	/** Source header filenames for all exported classes */
-	TArray<FString> AllSourceClassHeaders;
+	TSet<FString> AllSourceClassHeaders;
 	
 	/** Creats a 'glue' file that merges all generated script files */
 	void GlueAllGeneratedFiles();
diff --git a/Engine/Plugins/ScriptPlugin/Source/ScriptGeneratorPlugin/Private/LuaScriptCodeGenerator.cpp b/Engine/Plugins/ScriptPlugin/Source/ScriptGeneratorPlugin/Private/LuaScriptCodeGenerator.cpp
--- a/Engine/Plugins/ScriptPlugin/Source/ScriptGeneratorPlugin/Private/LuaScriptCodeGenerator.cpp
+++ b/Engine/Plugins/ScriptPlugin/Source/ScriptGeneratorPlugin/Private/LuaScriptCodeGenerator.cpp
@@ -218,7 +218,7 @@ bool FLuaScriptCodeGenerator::CanExportClass(UClass* Class)
 		// Check properties too
 		if (!bHasMembersToExport)
 		{
-			for (TFieldIterator<FProperty> PropertyIt(Class, EFieldIteratorFlags::ExcludeSuper); !bHasMembersToExport && PropertyIt; ++PropertyIt)
+			for (TFieldIterator<FProperty> PropertyIt(Class/*, EFieldIteratorFlags::ExcludeSuper*/); !bHasMembersToExport && PropertyIt; ++PropertyIt)
 			{
 				FProperty* Property = *PropertyIt;
 				if (CanExportProperty(ClassNameCPP, Class, Property))
@@ -254,13 +254,10 @@ FString FLuaScriptCodeGenerator::ExportFunction(const FString& ClassNameCPP, UCl
 	FProperty* ReturnValue = NULL;
 	UClass* FuncSuper = NULL;
 	
-	if (Function->GetOwnerClass() != Class)
+	// Find the base class where this function was defined
+	if (Function->GetOwnerClass() != Class && CanExportClass(Function->GetOwnerClass()))
 	{
-		// Find the base definition of the function
-		if (ExportedClasses.Contains(Function->GetOwnerClass()->GetFName()))
-		{
-			FuncSuper = Function->GetOwnerClass();
-		}
+		FuncSuper = Function->GetOwnerClass();
 	}
 
 	FString FunctionBody;
@@ -355,13 +352,10 @@ FString FLuaScriptCodeGenerator::ExportProperty(const FString& ClassNameCPP, UCl
 	FProperty* ReturnValue = NULL;
 	UClass* PropertySuper = NULL;
 
-	if (Property->GetOwnerClass() != Class)
+	// Find the base class where this property was defined
+	if (Property->GetOwnerClass() != Class && CanExportClass(Property->GetOwnerClass()))
 	{
-		// Find the base class where this property was defined
-		if (ExportedClasses.Contains(Property->GetOwnerClass()->GetFName()))
-		{
-			PropertySuper = Property->GetOwnerClass();
-		}
+		PropertySuper = Property->GetOwnerClass();
 	}
 
 	// Getter	
@@ -489,7 +483,7 @@ FString FLuaScriptCodeGenerator::ExportAdditionalClassGlue(const FString& ClassN
 
 void FLuaScriptCodeGenerator::ExportClass(UClass* Class, const FString& SourceHeaderFilename, const FString& GeneratedHeaderFilename, bool bHasChanged)
 {
-	if (!CanExportClass(Class))
+	if (!CanExportClass(Class) || ExportedClasses.Contains(Class->GetFName()))
 	{
 		return;
 	}
@@ -546,7 +540,15 @@ void FLuaScriptCodeGenerator::GlueAllGeneratedFiles()
 	FString LibGlue;
 
 	// Include all source header files
+	TArray<FString> Headers;
+	Headers.Reserve(AllSourceClassHeaders.Num());
 	for (auto& HeaderFilename : AllSourceClassHeaders)
+	{
+		Headers.Add(HeaderFilename);
+	}
+
+	Headers.Sort();
+	for (auto& HeaderFilename : Headers)
 	{
 		// Re-base to make sure we're including the right files on a remote machine
 		FString NewFilename(RebaseToBuildPath(HeaderFilename));
@@ -554,6 +556,7 @@ void FLuaScriptCodeGenerator::GlueAllGeneratedFiles()
 	}
 
 	// Include all script glue headers
+	AllScriptHeaders.Sort();
 	for (auto& HeaderFilename : AllScriptHeaders)
 	{
 		// Re-base to make sure we're including the right files on a remote machine
@@ -561,6 +564,7 @@ void FLuaScriptCodeGenerator::GlueAllGeneratedFiles()
 		LibGlue += FString::Printf(TEXT("#include \"%s\"\r\n"), *NewFilename);
 	}
 
+	LuaExportedClasses.Sort([](const UClass& x, const UClass& y) { return x.GetName() < y.GetName(); });
 	LibGlue += TEXT("\r\nvoid LuaRegisterExportedClasses(lua_State* InScriptContext)\r\n{\r\n");
 	for (auto Class : LuaExportedClasses)
 	{
diff --git a/Engine/Plugins/ScriptPlugin/Source/ScriptGeneratorPlugin/Private/LuaScriptCodeGenerator.h b/Engine/Plugins/ScriptPlugin/Source/ScriptGeneratorPlugin/Private/LuaScriptCodeGenerator.h
--- a/Engine/Plugins/ScriptPlugin/Source/ScriptGeneratorPlugin/Private/LuaScriptCodeGenerator.h
+++ b/Engine/Plugins/ScriptPlugin/Source/ScriptGeneratorPlugin/Private/LuaScriptCodeGenerator.h
@@ -20,7 +20,7 @@ protected:
 	/** All generated script header filenames */
 	TArray<FString> AllScriptHeaders;
 	/** Source header filenames for all exported classes */
-	TArray<FString> AllSourceClassHeaders;
+	TSet<FString> AllSourceClassHeaders;
 	/** All exported classes */
 	TArray<UClass*> LuaExportedClasses;
 	/** Functions exported for a class */
diff --git a/Engine/Plugins/ScriptPlugin/Source/ScriptGeneratorPlugin/Private/ScriptCodeGeneratorBase.cpp b/Engine/Plugins/ScriptPlugin/Source/ScriptGeneratorPlugin/Private/ScriptCodeGeneratorBase.cpp
--- a/Engine/Plugins/ScriptPlugin/Source/ScriptGeneratorPlugin/Private/ScriptCodeGeneratorBase.cpp
+++ b/Engine/Plugins/ScriptPlugin/Source/ScriptGeneratorPlugin/Private/ScriptCodeGeneratorBase.cpp
@@ -78,7 +78,9 @@ FString FScriptCodeGeneratorBase::GetPropertyTypeCPP(FProperty* Property, uint32
 	static const FString TEnumAsByteDecl(TEXT("TEnumAsByte<enum "));
 	static const FString TSubclassOfDecl(TEXT("TSubclassOf<class "));
 
-	FString PropertyType = Property->GetCPPType(NULL, PortFlags);
+	FString ExtendedType;
+	FString PropertyType = Property->GetCPPType(&ExtendedType, PortFlags);
+	PropertyType += ExtendedType;
 	// Strip any forward declaration keywords
 	if (PropertyType.StartsWith(EnumDecl, ESearchCase::CaseSensitive) || PropertyType.StartsWith(StructDecl, ESearchCase::CaseSensitive) || PropertyType.StartsWith(ClassDecl, ESearchCase::CaseSensitive))
 	{
@@ -102,7 +104,7 @@ FString FScriptCodeGeneratorBase::GenerateFunctionDispatch(UFunction* Function)
 {
 	FString Params;
 	
-	const bool bHasParamsOrReturnValue = (Function->Children != NULL);
+	const bool bHasParamsOrReturnValue = (Function->ChildProperties != NULL);
 	if (bHasParamsOrReturnValue)
 	{
 		Params += TEXT("\tstruct FDispatchParams\r\n\t{\r\n");
@@ -153,8 +155,7 @@ FString FScriptCodeGeneratorBase::GetScriptHeaderForClass(UClass* Class)
 
 bool FScriptCodeGeneratorBase::CanExportClass(UClass* Class)
 {
-	bool bCanExport = (Class->ClassFlags & (CLASS_RequiredAPI | CLASS_MinimalAPI)) && // Don't export classes that don't export DLL symbols
-		!ExportedClasses.Contains(Class->GetFName()); // Don't export classes that have already been exported
+	bool bCanExport = (Class->ClassFlags & (CLASS_RequiredAPI | CLASS_MinimalAPI)) != 0; // Don't export classes that don't export DLL symbols
 
 	return bCanExport;
 }
