UBT plugin version of the ScriptGenerator

#rnx

diff --git a/Engine/Plugins/ScriptPlugin/Source/ScriptGeneratorUbtPlugin/GenericScriptCodeGenerator.cs b/Engine/Plugins/ScriptPlugin/Source/ScriptGeneratorUbtPlugin/GenericScriptCodeGenerator.cs
new file mode 100644
--- /dev/null
+++ b/Engine/Plugins/ScriptPlugin/Source/ScriptGeneratorUbtPlugin/GenericScriptCodeGenerator.cs
@@ -0,0 +1,203 @@
+// Copyright Epic Games, Inc. All Rights Reserved.
+
+using EpicGames.Core;
+using EpicGames.UHT.Types;
+using EpicGames.UHT.Utils;
+using System;
+using System.Collections.Generic;
+using System.IO;
+using System.Linq;
+using System.Text;
+
+namespace ScriptGeneratorUbtPlugin
+{
+	static class GenericScriptCodeGeneratorStringBuilderExtensinos
+	{
+		public static StringBuilder AppendGenericWrapperFunctionDeclaration(this StringBuilder Builder, UhtClass Class, string FunctionName)
+		{
+			return Builder.Append("int32 ").Append(Class.EngineName).Append('_').Append(FunctionName).Append("(void* InScriptContext)");
+		}
+
+		public static StringBuilder AppendGenericFunctionParamDeclaration(this StringBuilder Builder, UhtClass Class, UhtProperty Property)
+		{
+			if (Property is UhtObjectPropertyBase)
+			{
+				Builder.Append("UObject* ").Append(Property.SourceName).Append(" = nullptr;");
+			}
+			else
+			{
+				Builder
+					.AppendPropertyText(Property, UhtPropertyTextType.GenericFunctionArgOrRetVal)
+					.Append(' ')
+					.Append(Property.SourceName)
+					.Append(" = ").AppendPropertyText(Property, UhtPropertyTextType.GenericFunctionArgOrRetVal)
+					.Append("();");
+			}
+			return Builder;
+		}
+
+		public static StringBuilder AppendGenericObjectDeclarationFromContext(this StringBuilder Builder, UhtClass Class)
+		{
+			return Builder.Append("UObject* Obj = (UObject*)InScriptContext;");
+		}
+
+		public static StringBuilder AppendGenericReturnValueHandler(this StringBuilder Builder, UhtClass Class, UhtProperty? ReturnValue)
+		{
+			return Builder.Append("return 0;");
+		}
+	}
+
+	internal class GenericScriptCodeGenerator : ScriptCodeGeneratorBase
+	{
+		public GenericScriptCodeGenerator(IUhtExportFactory Factory)
+			: base(Factory)
+		{
+		}
+
+		protected override bool CanExportClass(UhtClass Class)
+		{
+			if (!base.CanExportClass(Class))
+			{
+				return false;
+			}
+
+			foreach (UhtType Child in Class.Children)
+			{
+				if (Child is UhtFunction Function)
+				{
+					if (CanExportFunction(Class, Function))
+					{
+						return true;
+					}
+				}
+				else if (Child is UhtProperty Property)
+				{
+					if (CanExportProperty(Class, Property))
+					{
+						return true;
+					}
+				}
+			}
+			return false;
+		}
+
+		protected override void ExportClass(StringBuilder Builder, UhtClass Class)
+		{
+			Builder.Append("#pragma once\r\n\r\n");
+
+			List<UhtFunction> Functions = Class.Children.Where(x => x is UhtFunction).Cast<UhtFunction>().Reverse().ToList();
+
+			//ETSTODO - Functions are reversed in the engine
+			foreach (UhtFunction Function in Class.Functions.Reverse())
+			{
+				if (CanExportFunction(Class, Function))
+				{
+					ExportFunction(Builder, Class, Function);
+				}
+			}
+			//foreach (UhtType Child in Class.Children)
+			//{
+			//	if (Child is UhtFunction Function && CanExportFunction(Class, Function))
+			//	{
+			//		ExportFunction(Builder, Class, Function);
+			//	}
+			//}
+
+			foreach (UhtType Child in Class.Children)
+			{
+				if (Child is UhtProperty Property && CanExportProperty(Class, Property))
+				{
+					ExportProperty(Builder, Class, Property);
+				}
+			}
+
+			if (!Class.ClassFlags.HasAnyFlags(EClassFlags.Abstract))
+			{
+				Builder.AppendGenericWrapperFunctionDeclaration(Class, "New").Append("\r\n");
+				Builder.Append("{\r\n");
+				Builder.Append("\tUObject* Outer = NULL;\r\n");
+				Builder.Append("\tFName Name = FName(\"ScriptObject\");\r\n");
+				Builder.Append("\tUObject* Obj = NewObject<").Append(Class.SourceName).Append(">(Outer, Name);\r\n");
+				Builder.Append("\tif (Obj)\r\n\t{\r\n");
+				Builder.Append("\t\tFScriptObjectReferencer::Get().AddObjectReference(Obj);\r\n");
+				Builder.Append("\t\t// @todo: Register the object with the script context here\r\n");
+				Builder.Append("\t}\r\n\treturn 0;\r\n");
+				Builder.Append("}\r\n\r\n");
+
+				Builder.AppendGenericWrapperFunctionDeclaration(Class, "Destroy").Append("\r\n");
+				Builder.Append("{\r\n");
+				Builder.Append("\t").AppendGenericObjectDeclarationFromContext(Class).Append("\r\n");
+				Builder.Append("\tif (Obj)\r\n\t{\r\n");
+				Builder.Append("\t\tFScriptObjectReferencer::Get().RemoveObjectReference(Obj);\r\n");
+				Builder.Append("\t\t// @todo: Remove the object from the script context here if required\r\n");
+				Builder.Append("\t}\r\n\treturn 0;\r\n");
+				Builder.Append("}\r\n\r\n");
+			}
+		}
+
+		protected void ExportFunction(StringBuilder Builder, UhtClass Class, UhtFunction Function)
+		{
+			Builder.AppendGenericWrapperFunctionDeclaration(Class, Function.SourceName).Append("\r\n");
+			Builder.Append("{\r\n");
+			Builder.Append("\t").AppendGenericObjectDeclarationFromContext(Class).Append("\r\n");
+			AppendFunctionDispatch(Builder, Class, Function);
+			Builder.Append("\t").AppendGenericReturnValueHandler(Class, Function.ReturnProperty).Append("\r\n");
+			Builder.Append("}\r\n\r\n");
+		}
+
+		protected void ExportProperty(StringBuilder Builder, UhtClass Class, UhtProperty Property)
+		{
+			// Getter
+			Builder.AppendGenericWrapperFunctionDeclaration(Class, $"Get_{Property.SourceName}").Append("\r\n");
+			Builder.Append("{\r\n");
+			Builder.Append('\t').AppendGenericObjectDeclarationFromContext(Class).Append("\r\n");
+			Builder.Append("\tstatic FProperty* Property = FindScriptPropertyHelper(").Append(Class.SourceName).Append("::StaticClass(), TEXT(\"").Append(Property.SourceName).Append("\"));\r\n");
+			Builder.Append('\t').AppendGenericFunctionParamDeclaration(Class, Property).Append("\r\n");
+			Builder.Append("\tProperty->CopyCompleteValue(&").Append(Property.SourceName).Append(", Property->ContainerPtrToValuePtr<void>(Obj));\r\n");
+			Builder.Append("\t// @todo: handle property value here\r\n");
+			Builder.Append("\treturn 0;\r\n");
+			Builder.Append("}\r\n\r\n");
+
+			// Setter
+			Builder.AppendGenericWrapperFunctionDeclaration(Class, $"Set_{Property.SourceName}").Append("\r\n");
+			Builder.Append("{\r\n");
+			Builder.Append('\t').AppendGenericObjectDeclarationFromContext(Class).Append("\r\n");
+			Builder.Append("\tstatic FProperty* Property = FindScriptPropertyHelper(").Append(Class.SourceName).Append("::StaticClass(), TEXT(\"").Append(Property.SourceName).Append("\"));\r\n");
+			Builder.Append('\t').AppendGenericFunctionParamDeclaration(Class, Property).Append("\r\n");
+			Builder.Append("\tProperty->CopyCompleteValue(Property->ContainerPtrToValuePtr<void>(Obj), &").Append(Property.SourceName).Append(");\r\n");
+			Builder.Append("\treturn 0;\r\n");
+			Builder.Append("}\r\n\r\n");
+		}
+
+		protected override void Finish(StringBuilder Builder, List<UhtClass> Classes)
+		{
+			HashSet<UhtHeaderFile> UniqueHeaders = new HashSet<UhtHeaderFile>();
+			foreach (UhtClass Class in Classes)
+			{
+				UniqueHeaders.Add(Class.HeaderFile);
+			}
+			List<string> SortedHeaders = new List<string>();
+			foreach (UhtHeaderFile HeaderFile in UniqueHeaders)
+			{
+				SortedHeaders.Add(HeaderFile.FilePath);
+			}
+			SortedHeaders.Sort(StringComparerUE.OrdinalIgnoreCase);
+			foreach (string FilePath in SortedHeaders)
+			{
+				string RelativePath = Path.GetRelativePath(this.Factory.PluginModule!.IncludeBase, FilePath).Replace("\\", "/");
+				Builder.Append("#include \"").Append(RelativePath).Append("\"\r\n");
+			}
+
+			SortedHeaders.Clear();
+			foreach (UhtClass Class in Classes)
+			{
+				SortedHeaders.Add($"{Class.EngineName}.script.h");
+			}
+			SortedHeaders.Sort(StringComparerUE.OrdinalIgnoreCase);
+			foreach (string FilePath in SortedHeaders)
+			{
+				Builder.Append("#include \"").Append(FilePath).Append("\"\r\n");
+			}
+		}
+	}
+}
diff --git a/Engine/Plugins/ScriptPlugin/Source/ScriptGeneratorUbtPlugin/LuaScriptCodeGenerator.cs b/Engine/Plugins/ScriptPlugin/Source/ScriptGeneratorUbtPlugin/LuaScriptCodeGenerator.cs
new file mode 100644
--- /dev/null
+++ b/Engine/Plugins/ScriptPlugin/Source/ScriptGeneratorUbtPlugin/LuaScriptCodeGenerator.cs
@@ -0,0 +1,487 @@
+﻿// Copyright Epic Games, Inc. All Rights Reserved.
+
+using EpicGames.Core;
+using EpicGames.UHT.Types;
+using EpicGames.UHT.Utils;
+using System;
+using System.Collections.Generic;
+using System.Diagnostics;
+using System.IO;
+using System.Linq;
+using System.Text;
+
+namespace ScriptGeneratorUbtPlugin
+{
+	static class LuaScriptCodeGeneratorStringBuilderExtensinos
+	{
+		public static StringBuilder AppendLuaWrapperFunctionDeclaration(this StringBuilder Builder, UhtClass Class, string FunctionName)
+		{
+			return Builder.Append("int32 ").Append(Class.EngineName).Append('_').Append(FunctionName).Append("(lua_State* InScriptContext)");
+		}
+
+		public static StringBuilder AppendLuaObjectDeclarationFromContext(this StringBuilder Builder, UhtClass Class)
+		{
+			return Builder.Append("UObject* Obj = (").Append(Class.SourceName).Append("*)lua_touserdata(InScriptContext, 1);");
+		}
+
+		public static StringBuilder AppendLuaReturnValueHandler(this StringBuilder Builder, UhtClass Class, UhtProperty? ReturnValue, string? ReturnValueName)
+		{
+			if (ReturnValue != null)
+			{
+				if (ReturnValue is UhtIntProperty)
+				{
+					Builder.Append("lua_pushinteger(InScriptContext, ").Append(ReturnValueName).Append(");\r\n");
+				}
+				else if (ReturnValue is UhtFloatProperty)
+				{
+					Builder.Append("lua_pushnumber(InScriptContext, ").Append(ReturnValueName).Append(");\r\n");
+				}
+				else if (ReturnValue is UhtStrProperty)
+				{
+					Builder.Append("lua_pushstring(InScriptContext, TCHAR_TO_ANSI(*").Append(ReturnValueName).Append("));\r\n");
+				}
+				else if (ReturnValue is UhtNameProperty)
+				{
+					Builder.Append("lua_pushstring(InScriptContext, TCHAR_TO_ANSI(*").Append(ReturnValueName).Append(".ToString()));\r\n");
+				}
+				else if (ReturnValue is UhtBoolProperty)
+				{
+					Builder.Append("lua_pushboolean(InScriptContext, ").Append(ReturnValueName).Append(");\r\n");
+				}
+				else if (ReturnValue is UhtStructProperty StructProperty)
+				{
+					if (StructProperty.ScriptStruct.EngineName == "Vector2D")
+					{
+						Builder.Append("FLuaVector2D::Return(InScriptContext, ").Append(ReturnValueName).Append(");\r\n");
+					}
+					else if (StructProperty.ScriptStruct.EngineName == "Vector")
+					{
+						Builder.Append("FLuaVector::Return(InScriptContext, ").Append(ReturnValueName).Append(");\r\n");
+					}
+					else if (StructProperty.ScriptStruct.EngineName == "Vector4")
+					{
+						Builder.Append("FLuaVector4::Return(InScriptContext, ").Append(ReturnValueName).Append(");\r\n");
+					}
+					else if (StructProperty.ScriptStruct.EngineName == "Quat")
+					{
+						Builder.Append("FLuaQuat::Return(InScriptContext, ").Append(ReturnValueName).Append(");\r\n");
+					}
+					else if (StructProperty.ScriptStruct.EngineName == "LinearColor")
+					{
+						Builder.Append("FLuaLinearColor::Return(InScriptContext, ").Append(ReturnValueName).Append(");\r\n");
+					}
+					else if (StructProperty.ScriptStruct.EngineName == "Color")
+					{
+						Builder.Append("FLuaLinearColor::Return(InScriptContext, FLinearColor(").Append(ReturnValueName).Append("));\r\n");
+					}
+					else if (StructProperty.ScriptStruct.EngineName == "Transform")
+					{
+						Builder.Append("FLuaTransform::Return(InScriptContext, ").Append(ReturnValueName).Append(");\r\n");
+					}
+					else
+					{
+						throw new UhtIceException($"Unsupported function return value struct type: {StructProperty.ScriptStruct.EngineName}");
+					}
+				}
+				else if (ReturnValue is UhtObjectPropertyBase)
+				{
+					Builder.Append("lua_pushlightuserdata(InScriptContext, ").Append(ReturnValueName).Append(");\r\n");
+				}
+				else
+				{
+					throw new UhtIceException($"Unsupported function return type: {ReturnValue.GetType().Name}");
+				}
+				Builder.Append("\treturn 1;");
+			}
+			else
+			{
+				Builder.Append("return 0;");
+			}
+			return Builder;
+		}
+	}
+
+	internal class LuaScriptCodeGenerator : ScriptCodeGeneratorBase
+	{
+		public LuaScriptCodeGenerator(IUhtExportFactory Factory)
+			: base(Factory)
+		{
+		}
+
+		protected override bool CanExportClass(UhtClass Class)
+		{
+			if (!base.CanExportClass(Class))
+			{
+				return false;
+			}
+
+			for (UhtClass? Current = Class; Current != null; Current = Current.SuperClass)
+			{
+				foreach (UhtType Child in Current.Children)
+				{
+					if (Child is UhtFunction Function)
+					{
+						if (CanExportFunction(Class, Function))
+						{
+							return true;
+						}
+					}
+				}
+			}
+
+			foreach (UhtType Child in Class.Children)
+			{
+				if (Child is UhtProperty Property)
+				{
+					if (CanExportProperty(Class, Property))
+					{
+						return true;
+					}
+				}
+			}
+			return false;
+		}
+
+		protected override bool CanExportFunction(UhtClass Class, UhtFunction Function)
+		{
+			if (!base.CanExportFunction(Class, Function))
+			{
+				return false;
+			}
+
+			foreach (UhtType Child in Function.Children)
+			{
+				if (Child is UhtProperty Property)
+				{
+					if (!IsPropertyTypeSupported(Property))
+					{
+						return false;
+					}
+				}
+			}
+			return true;
+		}
+
+		protected override bool CanExportProperty(UhtClass Class, UhtProperty Property)
+		{
+			return Property.PropertyFlags.HasAnyFlags(EPropertyFlags.Edit) && IsPropertyTypeSupported(Property);
+		}
+
+		private bool IsPropertyTypeSupported(UhtProperty Property)
+		{
+			if (Property is UhtStructProperty StructProperty)
+			{
+				return
+					StructProperty.ScriptStruct.EngineName == "Vector2D" ||
+					StructProperty.ScriptStruct.EngineName == "Vector" ||
+					StructProperty.ScriptStruct.EngineName == "Vector4" ||
+					StructProperty.ScriptStruct.EngineName == "Quat" ||
+					StructProperty.ScriptStruct.EngineName == "LinearColor" ||
+					StructProperty.ScriptStruct.EngineName == "Color" ||
+					StructProperty.ScriptStruct.EngineName == "Transform";
+			}
+			else if (
+				Property is UhtLazyObjectPtrProperty ||
+				Property is UhtSoftObjectProperty ||
+				Property is UhtSoftClassProperty ||
+				Property is UhtWeakObjectPtrProperty)
+			{
+				return false;
+			}
+			else if (
+				Property is UhtIntProperty ||
+				Property is UhtFloatProperty ||
+				Property is UhtStrProperty ||
+				Property is UhtNameProperty ||
+				Property is UhtBoolProperty ||
+				Property is UhtObjectPropertyBase)
+			{
+				return true;
+			}
+			return false;
+		}
+
+		protected override void ExportClass(StringBuilder Builder, UhtClass Class)
+		{
+			Builder.Append("#pragma once\r\n\r\n");
+
+			List<UhtFunction> Functions = Class.Children.Where(x => x is UhtFunction).Cast<UhtFunction>().Reverse().ToList();
+
+			//ETSTODO - Functions are reversed in the engine
+			for (UhtClass? Current = Class; Current != null; Current = Current.SuperClass)
+			{
+				foreach (UhtFunction Function in Current.Functions.Reverse())
+				{
+					if (CanExportFunction(Class, Function))
+					{
+						ExportFunction(Builder, Class, Function);
+					}
+				}
+			}
+
+			for (UhtClass? Current = Class; Current != null; Current = Current.SuperClass)
+			{
+				foreach (UhtType Child in Current.Children)
+				{
+					if (Child is UhtProperty Property && CanExportProperty(Class, Property))
+					{
+						ExportProperty(Builder, Class, Property);
+					}
+				}
+			}
+
+			if (!Class.ClassFlags.HasAnyFlags(EClassFlags.Abstract))
+			{
+				Builder.AppendLuaWrapperFunctionDeclaration(Class, "New").Append("\r\n");
+				Builder.Append("{\r\n");
+				Builder.Append("\tUObject* Outer = (UObject*)lua_touserdata(InScriptContext, 1);\r\n");
+				Builder.Append("\tFName Name = FName(luaL_checkstring(InScriptContext, 2));\r\n");
+				Builder.Append("\tUObject* Obj = NewObject<").Append(Class.SourceName).Append(">(Outer, Name);\r\n");
+				Builder.Append("\tif (Obj)\r\n\t{\r\n");
+				Builder.Append("\t\tFScriptObjectReferencer::Get().AddObjectReference(Obj);\r\n");
+				Builder.Append("\t}\r\n");
+				Builder.Append("\tlua_pushlightuserdata(InScriptContext, Obj);\r\n");
+				Builder.Append("\treturn 1;\r\n");
+				Builder.Append("}\r\n\r\n");
+
+				Builder.AppendLuaWrapperFunctionDeclaration(Class, "Destroy").Append("\r\n");
+				Builder.Append("{\r\n");
+				Builder.Append("\t").AppendLuaObjectDeclarationFromContext(Class).Append("\r\n");
+				Builder.Append("\tif (Obj)\r\n\t{\r\n");
+				Builder.Append("\t\tFScriptObjectReferencer::Get().RemoveObjectReference(Obj);\r\n");
+				Builder.Append("\t}\r\n");
+				Builder.Append("\treturn 0;\r\n");
+				Builder.Append("}\r\n\r\n");
+			}
+
+			// Class: Equivalent of StaticClass()
+			Builder.AppendLuaWrapperFunctionDeclaration(Class, "Class").Append("\r\n");
+			Builder.Append("{\r\n");
+			Builder.Append("\tUClass* Class = ").Append(Class.SourceName).Append("::StaticClass();\r\n");
+			Builder.Append("\tlua_pushlightuserdata(InScriptContext, Class);\r\n");
+			Builder.Append("\treturn 1;\r\n");
+			Builder.Append("}\r\n\r\n");
+
+			// Library
+			Builder.Append("static const luaL_Reg ").Append(Class.EngineName).Append("_Lib[] =\r\n");
+			Builder.Append("{\r\n");
+			if (!Class.ClassFlags.HasAnyFlags(EClassFlags.Abstract))
+			{
+				Builder.Append("\t{ \"New\", ").Append(Class.EngineName).Append("_New },\r\n");
+				Builder.Append("\t{ \"Destroy\", ").Append(Class.EngineName).Append("_Destroy },\r\n");
+				Builder.Append("\t{ \"Class\", ").Append(Class.EngineName).Append("_Class },\r\n");
+			}
+
+			//ETSTODO - Functions are reversed in the engine
+			for (UhtClass? Current = Class; Current != null; Current = Current.SuperClass)
+			{
+				foreach (UhtFunction Function in Current.Functions.Reverse())
+				{
+					if (CanExportFunction(Class, Function))
+					{
+						Builder.Append("\t{ \"").Append(Function.SourceName).Append("\", ").Append(Class.EngineName).Append('_').Append(Function.SourceName).Append(" },\r\n");
+					}
+				}
+			}
+
+			for (UhtClass? Current = Class; Current != null; Current = Current.SuperClass)
+			{
+				foreach (UhtType Child in Current.Children)
+				{
+					if (Child is UhtProperty Property && CanExportProperty(Class, Property))
+					{
+						Builder.Append("\t{ \"Get_").Append(Property.SourceName).Append("\", ").Append(Class.EngineName).Append("_Get_").Append(Property.SourceName).Append(" },\r\n");
+						Builder.Append("\t{ \"Set_").Append(Property.SourceName).Append("\", ").Append(Class.EngineName).Append("_Set_").Append(Property.SourceName).Append(" },\r\n");
+					}
+				}
+			}
+
+			Builder.Append("\t{ NULL, NULL }\r\n");
+			Builder.Append("};\r\n\r\n");
+		}
+
+		protected void ExportFunction(StringBuilder Builder, UhtClass Class, UhtFunction Function)
+		{
+			UhtClass? FunctionSuper = null;
+			if (Function.Outer != null && Function.Outer != Class && Function.Outer is UhtClass OuterClass && base.CanExportClass(OuterClass))
+			{
+				FunctionSuper = OuterClass;
+			}
+
+			Builder.AppendLuaWrapperFunctionDeclaration(Class, Function.SourceName).Append("\r\n");
+			Builder.Append("{\r\n");
+			if (FunctionSuper == null)
+			{
+				Builder.Append("\t").AppendLuaObjectDeclarationFromContext(Class).Append("\r\n");
+				AppendFunctionDispatch(Builder, Class, Function);
+				string ReturnValueName = Function.ReturnProperty != null ? $"Params.{Function.ReturnProperty.SourceName}" : String.Empty;
+				Builder.Append("\t").AppendLuaReturnValueHandler(Class, Function.ReturnProperty, ReturnValueName).Append("\r\n");
+			}
+			else
+			{
+				Builder.Append("\treturn ").Append(FunctionSuper.EngineName).Append('_').Append(Function.SourceName).Append("(InScriptContext);\r\n");
+			}
+			Builder.Append("}\r\n\r\n");
+		}
+
+		protected void ExportProperty(StringBuilder Builder, UhtClass Class, UhtProperty Property)
+		{
+			UhtClass? PropertySuper = null;
+			if (Property.Outer != null && Property.Outer != Class && Property.Outer is UhtClass OuterClass && base.CanExportClass(OuterClass))
+			{
+				PropertySuper = OuterClass;
+			}
+
+			// Getter
+			Builder.AppendLuaWrapperFunctionDeclaration(Class, $"Get_{Property.SourceName}").Append("\r\n");
+			Builder.Append("{\r\n");
+			if (PropertySuper == null)
+			{
+				Builder.Append('\t').AppendLuaObjectDeclarationFromContext(Class).Append("\r\n");
+				Builder.Append("\tstatic FProperty* Property = FindScriptPropertyHelper(").Append(Class.SourceName).Append("::StaticClass(), TEXT(\"").Append(Property.SourceName).Append("\"));\r\n");
+				Builder.Append('\t').AppendPropertyText(Property, UhtPropertyTextType.GenericFunctionArgOrRetVal).Append(" PropertyValue;\r\n");
+				Builder.Append("\tProperty->CopyCompleteValue(&PropertyValue, Property->ContainerPtrToValuePtr<void>(Obj));\r\n");
+				Builder.Append('\t').AppendLuaReturnValueHandler(Class, Property, "PropertyValue").Append("\r\n");
+			}
+			else
+			{
+				Builder.Append("\treturn ").Append(PropertySuper.EngineName).Append('_').Append("Get_").Append(Property.SourceName).Append("(InScriptContext);\r\n");
+			}
+			Builder.Append("}\r\n\r\n");
+
+			// Setter
+			Builder.AppendLuaWrapperFunctionDeclaration(Class, $"Set_{Property.SourceName}").Append("\r\n");
+			Builder.Append("{\r\n");
+			if (PropertySuper == null)
+			{
+				Builder.Append('\t').AppendLuaObjectDeclarationFromContext(Class).Append("\r\n");
+				Builder.Append("\tstatic FProperty* Property = FindScriptPropertyHelper(").Append(Class.SourceName).Append("::StaticClass(), TEXT(\"").Append(Property.SourceName).Append("\"));\r\n");
+				Builder.Append('\t').AppendPropertyText(Property, UhtPropertyTextType.GenericFunctionArgOrRetVal).Append(" PropertyValue = ");
+				AppendInitializeFunctionDispatchParam(Builder, Class, null, Property, 0).Append(";\r\n");
+				Builder.Append("\tProperty->CopyCompleteValue(Property->ContainerPtrToValuePtr<void>(Obj), &PropertyValue);\r\n");
+				Builder.Append("\treturn 0;\r\n");
+			}
+			else
+			{
+				Builder.Append("\treturn ").Append(PropertySuper.EngineName).Append('_').Append("Set_").Append(Property.SourceName).Append("(InScriptContext);\r\n");
+			}
+			Builder.Append("}\r\n\r\n");
+		}
+
+		protected override void Finish(StringBuilder Builder, List<UhtClass> Classes)
+		{
+			HashSet<UhtHeaderFile> UniqueHeaders = new HashSet<UhtHeaderFile>();
+			foreach (UhtClass Class in Classes)
+			{
+				UniqueHeaders.Add(Class.HeaderFile);
+			}
+			List<string> SortedHeaders = new List<string>();
+			foreach (UhtHeaderFile HeaderFile in UniqueHeaders)
+			{
+				SortedHeaders.Add(HeaderFile.FilePath);
+			}
+			SortedHeaders.Sort(StringComparerUE.OrdinalIgnoreCase);
+			foreach (string FilePath in SortedHeaders)
+			{
+				string RelativePath = Path.GetRelativePath(this.Factory.PluginModule!.IncludeBase, FilePath).Replace("\\", "/");
+				Builder.Append("#include \"").Append(RelativePath).Append("\"\r\n");
+			}
+
+			Classes.Sort((x, y) => StringComparerUE.OrdinalIgnoreCase.Compare(x.EngineName, y.EngineName));
+			SortedHeaders.Clear();
+			foreach (UhtClass Class in Classes)
+			{
+				Builder.Append("#include \"").Append(Class.EngineName).Append(".script.h\"\r\n");
+			}
+			Builder.Append("\r\n");
+			Builder.Append("void LuaRegisterExportedClasses(lua_State* InScriptContext)\r\n");
+			Builder.Append("{\r\n");
+			foreach(UhtClass Class in Classes)
+			{
+				Builder.Append("\tFLuaUtils::RegisterLibrary(InScriptContext, ").Append(Class.EngineName).Append("_Lib, \"").Append(Class.EngineName).Append("\");\r\n");
+			}
+			Builder.Append("}\r\n\r\n");
+		}
+
+		protected override StringBuilder AppendInitializeFunctionDispatchParam(StringBuilder Builder, UhtClass Class, UhtFunction? Function, UhtProperty Property, int PropertyIndex)
+		{
+			if (!Property.PropertyFlags.HasAnyFlags(EPropertyFlags.ReturnParm))
+			{
+				int ParamIndex = PropertyIndex + 2;
+				if (Property is UhtIntProperty)
+				{
+					Builder.Append("(luaL_checkint");
+				}
+				else if (Property is UhtFloatProperty)
+				{
+					Builder.Append("(float)(luaL_checknumber");
+				}
+				else if (Property is UhtStrProperty)
+				{
+					Builder.Append("ANSI_TO_TCHAR(luaL_checkstring");
+				}
+				else if (Property is UhtNameProperty)
+				{
+					Builder.Append("FName(luaL_checkstring");
+				}
+				else if (Property is UhtBoolProperty)
+				{
+					Builder.Append("!!(lua_toboolean");
+				}
+				else if (Property is UhtStructProperty StructProperty)
+				{
+					if (StructProperty.ScriptStruct.EngineName == "Vector2D")
+					{
+						Builder.Append("(FLuaVector2D::Get");
+					}
+					else if (StructProperty.ScriptStruct.EngineName == "Vector")
+					{
+						Builder.Append("(FLuaVector::Get");
+					}
+					else if (StructProperty.ScriptStruct.EngineName == "Vector4")
+					{
+						Builder.Append("(FLuaVector4::Get");
+					}
+					else if (StructProperty.ScriptStruct.EngineName == "Quat")
+					{
+						Builder.Append("(FLuaQuat::Get");
+					}
+					else if (StructProperty.ScriptStruct.EngineName == "LinearColor")
+					{
+						Builder.Append("(FLuaLinearColor::Get");
+					}
+					else if (StructProperty.ScriptStruct.EngineName == "Color")
+					{
+						Builder.Append("FColor(FLuaLinearColor::Get");
+					}
+					else if (StructProperty.ScriptStruct.EngineName == "Transform")
+					{
+						Builder.Append("(FLuaTransform::Get");
+					}
+					else
+					{
+						throw new UhtIceException($"Unsupported function param struct type: {StructProperty.ScriptStruct.EngineName}");
+					}
+				}
+				else if (Property is UhtClassProperty)
+				{
+					Builder.Append("(UClass*)(lua_touserdata");
+				}
+				else if (Property is UhtObjectPropertyBase)
+				{
+					Builder.Append("(").AppendPropertyText(Property, UhtPropertyTextType.GenericFunctionArgOrRetVal).Append(")(lua_touserdata");
+				}
+				else
+				{
+					throw new UhtIceException($"Unsupported function param type: {Property.GetType().Name}");
+				}
+				Builder.Append("(InScriptContext, ").Append(ParamIndex).Append("))");
+			}
+			else
+			{
+				base.AppendInitializeFunctionDispatchParam(Builder, Class, Function, Property, PropertyIndex);
+			}
+			return Builder;
+		}
+	}
+}
diff --git a/Engine/Plugins/ScriptPlugin/Source/ScriptGeneratorUbtPlugin/ScriptCodeGeneratorBase.cs b/Engine/Plugins/ScriptPlugin/Source/ScriptGeneratorUbtPlugin/ScriptCodeGeneratorBase.cs
new file mode 100644
--- /dev/null
+++ b/Engine/Plugins/ScriptPlugin/Source/ScriptGeneratorUbtPlugin/ScriptCodeGeneratorBase.cs
@@ -0,0 +1,244 @@
+// Copyright Epic Games, Inc. All Rights Reserved.
+
+using EpicGames.Core;
+using EpicGames.UHT.Types;
+using EpicGames.UHT.Utils;
+using System;
+using System.Collections.Generic;
+using System.Linq;
+using System.Text;
+using System.Threading.Tasks;
+using UnrealBuildBase;
+using UnrealBuildTool;
+
+namespace ScriptGeneratorUbtPlugin
+{
+	abstract internal class ScriptCodeGeneratorBase
+	{
+		public readonly IUhtExportFactory Factory;
+		public UhtSession Session => this.Factory.Session;
+
+		public ScriptCodeGeneratorBase(IUhtExportFactory Factory)
+		{
+			this.Factory = Factory;
+		}
+
+		/// <summary>
+		/// Export all the classes in all the packages
+		/// </summary>
+		public void Generate()
+		{
+			DirectoryReference ConfigDirectory = DirectoryReference.Combine(Unreal.EngineDirectory, "Programs/UnrealBuildTool");
+			ConfigHierarchy Ini = ConfigCache.ReadHierarchy(ConfigHierarchyType.Engine, ConfigDirectory, BuildHostPlatform.Current.Platform);
+			List<string>? SupportedScriptModules = null;
+			Ini.GetArray("Plugins", "ScriptSupportedModules", out SupportedScriptModules);
+
+			// Loop through the packages making sure they should be exported.  Queue the export of the classes
+			List<UhtClass> Classes = new List<UhtClass>();
+			List<Task?> Tasks = new List<Task?>();
+			foreach (UhtPackage Package in this.Session.Packages)
+			{
+				if (Package.Module.ModuleType != UHTModuleType.EngineRuntime && Package.Module.ModuleType != UHTModuleType.GameRuntime)
+				{
+					continue;
+				}
+
+				if (SupportedScriptModules != null && !SupportedScriptModules.Any(x => string.Compare(x, Package.Module.Name, StringComparison.OrdinalIgnoreCase) == 0))
+				{
+					continue;
+				}
+
+				QueueClassExports(Package, Package, Classes, Tasks);
+			}
+
+			// Wait for all the classes to export
+			Task[]? WaitTasks = Tasks.Where(x => x != null).Cast<Task>().ToArray();
+			if (WaitTasks.Length > 0)
+			{
+				Task.WaitAll(WaitTasks);
+			}
+
+			// Finish the export process
+			Finish(Classes);
+		}
+
+		/// <summary>
+		/// Collect the classes to be exported for the given package and type
+		/// </summary>
+		/// <param name="Package">Package being exported</param>
+		/// <param name="Type">Type to test for exporting</param>
+		/// <param name="Classes">Collection of exported classes</param>
+		/// <param name="Tasks">Collection of queued tasks</param>
+		private void QueueClassExports(UhtPackage Package, UhtType Type, List<UhtClass> Classes, List<Task?> Tasks)
+		{
+			if (Type is UhtClass Class)
+			{
+				if (CanExportClass(Class))
+				{
+					Classes.Add(Class);
+					Tasks.Add(Factory.CreateTask((Factory) => { ExportClass(Package, Class); }));
+				}
+			}
+			foreach (UhtType Child in Type.Children)
+			{
+				QueueClassExports(Package, Child, Classes, Tasks);
+			}
+		}
+
+		/// <summary>
+		/// Test to see if the given class should be exported
+		/// </summary>
+		/// <param name="Class">Class to test</param>
+		/// <returns>True if the class should be exported, false if not</returns>
+		protected virtual bool CanExportClass(UhtClass Class)
+		{
+			return Class.ClassFlags.HasAnyFlags(EClassFlags.RequiredAPI | EClassFlags.MinimalAPI); // Don't export classes that don't export DLL symbols
+		}
+
+		/// <summary>
+		/// Test to see if the given function should be exported
+		/// </summary>
+		/// <param name="Class">Owning class of the function</param>
+		/// <param name="Function">Function to test</param>
+		/// <returns>True if the function should be exported</returns>
+		protected virtual bool CanExportFunction(UhtClass Class, UhtFunction Function)
+		{
+			// We don't support delegates and non-public functions
+			if (Function.FunctionFlags.HasAnyFlags(EFunctionFlags.Delegate))
+			{
+				return false;
+			}
+
+			// Reject if any of the parameter types is unsupported yet
+			foreach (UhtType Child in Function.Children)
+			{
+				if (Child is UhtArrayProperty ||
+					Child is UhtDelegateProperty ||
+					Child is UhtMulticastDelegateProperty ||
+					Child is UhtWeakObjectPtrProperty ||
+					Child is UhtInterfaceProperty)
+				{
+					return false;
+				}
+				if (Child is UhtProperty Property && Property.bIsStaticArray)
+				{
+					return false;
+				}
+			}
+			return true;
+		}
+
+		/// <summary>
+		/// Test to see if the given property should be exported
+		/// </summary>
+		/// <param name="Class">Owning class of the property</param>
+		/// <param name="Property">Property to test</param>
+		/// <returns>True if the property should be exported</returns>
+		protected virtual bool CanExportProperty(UhtClass Class, UhtProperty Property)
+		{
+			// Property must be DLL exported
+			if (!Class.ClassFlags.HasAnyFlags(EClassFlags.RequiredAPI))
+			{
+				return false;
+			}
+
+			// Only public, editable properties can be exported
+			if (Property.PropertyFlags.HasAnyFlags(EPropertyFlags.NativeAccessSpecifierPrivate | EPropertyFlags.NativeAccessSpecifierProtected) || 
+				!Property.PropertyFlags.HasAnyFlags(EPropertyFlags.Edit))
+			{
+				return false;
+			}
+
+			// Reject if any of the parameter types is unsupported yet
+			if (Property.bIsStaticArray ||
+				Property is UhtArrayProperty ||
+				Property is UhtDelegateProperty ||
+				Property is UhtMulticastDelegateProperty ||
+				Property is UhtWeakObjectPtrProperty ||
+				Property is UhtInterfaceProperty ||
+				Property is UhtStructProperty)
+			{
+				return false;
+			}
+			return true;
+		}
+
+		/// <summary>
+		/// Export the given class
+		/// </summary>
+		/// <param name="Factory">Factory associated with the export</param>
+		/// <param name="Package">Owning package</param>
+		/// <param name="Class">Class to export</param>
+		private void ExportClass(UhtPackage Package, UhtClass Class)
+		{
+			using (BorrowStringBuilder Borrower = new BorrowStringBuilder(StringBuilderCache.Big))
+			{
+				ExportClass(Borrower.StringBuilder, Class);
+				string FileName = this.Factory.MakePath(Class.EngineName, ".script.h"); //ETSTODO - WRONG
+				this.Factory.CommitOutput(FileName, Borrower.StringBuilder);
+			}
+		}
+
+		private void Finish(List<UhtClass> Classes)
+		{
+			using (BorrowStringBuilder Borrower = new BorrowStringBuilder(StringBuilderCache.Big))
+			{
+				Finish(Borrower.StringBuilder, Classes);
+				string FileName = this.Factory.MakePath("GeneratedScriptLibraries", ".inl"); //ETSTODO - WRONG
+				this.Factory.CommitOutput(FileName, Borrower.StringBuilder);
+			}
+		}
+
+		abstract protected void ExportClass(StringBuilder Builder, UhtClass Class);
+
+		abstract protected void Finish(StringBuilder Builder, List<UhtClass> Classes);
+
+		virtual protected StringBuilder AppendInitializeFunctionDispatchParam(StringBuilder Builder, UhtClass Class, UhtFunction? Function, UhtProperty Property, int PropertyIndex)
+		{
+			if (Property is UhtObjectPropertyBase)
+			{
+				Builder.Append("NULL");
+			}
+			else
+			{
+				Builder.AppendPropertyText(Property, UhtPropertyTextType.GenericFunctionArgOrRetVal).Append("()");
+			}
+			return Builder;
+		}
+
+		virtual protected StringBuilder AppendFunctionDispatch(StringBuilder Builder, UhtClass Class, UhtFunction Function)
+		{
+			bool bHasParamsOrReturnValue = Function.Children.Count > 0;
+			if (bHasParamsOrReturnValue)
+			{
+				Builder.Append("\tstruct FDispatchParams\r\n");
+				Builder.Append("\t{\r\n");
+				foreach (UhtProperty Property in Function.Children)
+				{
+					Builder.Append("\t\t").AppendPropertyText(Property, UhtPropertyTextType.GenericFunctionArgOrRetVal).Append(' ').Append(Property.SourceName).Append(";\r\n");
+				}
+				Builder.Append("\t} Params;\r\n");
+				int PropertyIndex = 0;
+				foreach (UhtProperty Property in Function.Children)
+				{
+					Builder.Append("\tParams.").Append(Property.SourceName).Append(" = ");
+					AppendInitializeFunctionDispatchParam(Builder, Class, Function, Property, PropertyIndex).Append(";\r\n");
+					PropertyIndex++;
+				}
+			}
+
+			Builder.Append("\tstatic UFunction* Function = Obj->FindFunctionChecked(TEXT(\"").Append(Function.SourceName).Append("\"));\r\n");
+
+			if (bHasParamsOrReturnValue)
+			{
+				Builder.Append("\tcheck(Function->ParmsSize == sizeof(FDispatchParams));\r\n");
+				Builder.Append("\tObj->ProcessEvent(Function, &Params);\r\n");
+			}
+			else
+			{
+				Builder.Append("\tObj->ProcessEvent(Function, NULL);\r\n");
+			}
+			return Builder;
+		}
+	}
+}
diff --git a/Engine/Plugins/ScriptPlugin/Source/ScriptGeneratorUbtPlugin/ScriptGenerator.cs b/Engine/Plugins/ScriptPlugin/Source/ScriptGeneratorUbtPlugin/ScriptGenerator.cs
new file mode 100644
--- /dev/null
+++ b/Engine/Plugins/ScriptPlugin/Source/ScriptGeneratorUbtPlugin/ScriptGenerator.cs
@@ -0,0 +1,39 @@
+// Copyright Epic Games, Inc. All Rights Reserved.
+
+using EpicGames.UHT.Tables;
+using EpicGames.UHT.Utils;
+
+namespace ScriptGeneratorUbtPlugin
+{
+	[UnrealHeaderTool]
+    class ScriptGenerator
+	{
+		[UhtExporter(Name = "ScriptPlugin", Description = "Generic Script Plugin Generator", Options = UhtExporterOptions.Default, ModuleName="ScriptPlugin")]
+		private static void ScriptGeneratorExporter(IUhtExportFactory Factory)
+		{
+
+			// Make sure this plugin should be run
+			if (!Factory.Session.IsPluginEnabled("ScriptPlugin", false))
+			{
+				return;
+			}
+
+			// Based on the WITH_LUA setting, run the proper exporter.
+			if (Factory.PluginModule != null)
+			{
+				int Value;
+				if (Factory.PluginModule.TryGetDefine("WITH_LUA", out Value))
+				{
+					if (Value == 0)
+					{
+						new GenericScriptCodeGenerator(Factory).Generate();
+					}
+					else
+					{
+						new LuaScriptCodeGenerator(Factory).Generate();
+					}
+				}
+			}
+		}
+	}
+}
diff --git a/Engine/Plugins/ScriptPlugin/Source/ScriptGeneratorUbtPlugin/ScriptGeneratorUbtPlugin.ubtplugin.csproj b/Engine/Plugins/ScriptPlugin/Source/ScriptGeneratorUbtPlugin/ScriptGeneratorUbtPlugin.ubtplugin.csproj
new file mode 100644
--- /dev/null
+++ b/Engine/Plugins/ScriptPlugin/Source/ScriptGeneratorUbtPlugin/ScriptGeneratorUbtPlugin.ubtplugin.csproj
@@ -0,0 +1,50 @@
+﻿<Project Sdk="Microsoft.NET.Sdk">
+  <Import Project="ScriptGeneratorUbtPlugin.ubtplugin.csproj.props" Condition="Exists('ScriptGeneratorUbtPlugin.ubtplugin.csproj.props')" />
+  <Import Project="$(EngineDir)\Source\Programs\Shared\UnrealEngine.csproj.props" />
+  
+  <PropertyGroup>
+    <TargetFramework>net6.0</TargetFramework>
+    <Configuration Condition=" '$(Configuration)' == '' ">Development</Configuration>
+    <OutputType>Library</OutputType>
+    <AppendTargetFrameworkToOutputPath>false</AppendTargetFrameworkToOutputPath>
+    <GenerateAssemblyInfo>false</GenerateAssemblyInfo>
+    <GenerateTargetFrameworkAttribute>false</GenerateTargetFrameworkAttribute>
+    <Configurations>Debug;Release;Development</Configurations>
+    <RootNamespace>ScriptGeneratorUbtPlugin</RootNamespace>
+    <AssemblyName>ScriptGeneratorUbtPlugin</AssemblyName>
+    <OutputPath>..\..\..\..\Binaries\DotNET\UnrealBuildTool\Plugins\ScriptGeneratorUbtPlugin\</OutputPath>
+    <TreatWarningsAsErrors>true</TreatWarningsAsErrors>
+    <DebugType>pdbonly</DebugType>
+    <Nullable>enable</Nullable>
+  </PropertyGroup>
+  <PropertyGroup Condition=" '$(Configuration)|$(Platform)' == 'Development|AnyCPU' ">
+    <DefineConstants>$(DefineConstants);TRACE</DefineConstants>
+    <Optimize>true</Optimize>
+  </PropertyGroup>
+  <PropertyGroup Condition=" '$(Configuration)|$(Platform)' == 'Release|AnyCPU' ">
+    <DefineConstants>$(DefineConstants);TRACE</DefineConstants>
+    <Optimize>true</Optimize>
+  </PropertyGroup>
+  <PropertyGroup Condition=" '$(Configuration)|$(Platform)' == 'Debug|AnyCPU' ">
+    <DefineConstants>$(DefineConstants);DEBUG;TRACE</DefineConstants>
+    <DebugSymbols>true</DebugSymbols>
+    <DebugType>full</DebugType>
+  </PropertyGroup>
+  <ItemGroup>
+    <PackageReference Include="Microsoft.CSharp" Version="4.7.0" />
+  </ItemGroup>
+  <ItemGroup>
+    <Reference Include="EpicGames.Build">
+      <HintPath>$(EngineDir)\Binaries\DotNET\UnrealBuildTool\EpicGames.Build.dll</HintPath>
+    </Reference>
+    <Reference Include="EpicGames.Core">
+      <HintPath>$(EngineDir)\Binaries\DotNET\UnrealBuildTool\EpicGames.Core.dll</HintPath>
+    </Reference>
+    <Reference Include="EpicGames.UHT">
+      <HintPath>$(EngineDir)\Binaries\DotNET\UnrealBuildTool\EpicGames.UHT.dll</HintPath>
+    </Reference>
+    <Reference Include="UnrealBuildTool">
+      <HintPath>$(EngineDir)\Binaries\DotNET\UnrealBuildTool\UnrealBuildTool.dll</HintPath>
+    </Reference>
+  </ItemGroup>
+</Project>
\ No newline at end of file
