More coding style cleanups

#rnx

diff --git a/Engine/Plugins/ScriptPlugin/Source/ScriptGeneratorUbtPlugin/GenericScriptCodeGenerator.cs b/Engine/Plugins/ScriptPlugin/Source/ScriptGeneratorUbtPlugin/GenericScriptCodeGenerator.cs
--- a/Engine/Plugins/ScriptPlugin/Source/ScriptGeneratorUbtPlugin/GenericScriptCodeGenerator.cs
+++ b/Engine/Plugins/ScriptPlugin/Source/ScriptGeneratorUbtPlugin/GenericScriptCodeGenerator.cs
@@ -13,66 +13,66 @@ namespace ScriptGeneratorUbtPlugin
 {
 	static class GenericScriptCodeGeneratorStringBuilderExtensinos
 	{
-		public static StringBuilder AppendGenericWrapperFunctionDeclaration(this StringBuilder Builder, UhtClass Class, string FunctionName)
+		public static StringBuilder AppendGenericWrapperFunctionDeclaration(this StringBuilder builder, UhtClass classObj, string functionName)
 		{
-			return Builder.Append("int32 ").Append(Class.EngineName).Append('_').Append(FunctionName).Append("(void* InScriptContext)");
+			return builder.Append("int32 ").Append(classObj.EngineName).Append('_').Append(functionName).Append("(void* InScriptContext)");
 		}
 
-		public static StringBuilder AppendGenericFunctionParamDeclaration(this StringBuilder Builder, UhtClass Class, UhtProperty Property)
+		public static StringBuilder AppendGenericFunctionParamDeclaration(this StringBuilder builder, UhtClass _ /* classObj */, UhtProperty property)
 		{
-			if (Property is UhtObjectPropertyBase)
+			if (property is UhtObjectPropertyBase)
 			{
-				Builder.Append("UObject* ").Append(Property.SourceName).Append(" = nullptr;");
+				builder.Append("UObject* ").Append(property.SourceName).Append(" = nullptr;");
 			}
 			else
 			{
-				Builder
-					.AppendPropertyText(Property, UhtPropertyTextType.GenericFunctionArgOrRetVal)
+				builder
+					.AppendPropertyText(property, UhtPropertyTextType.GenericFunctionArgOrRetVal)
 					.Append(' ')
-					.Append(Property.SourceName)
-					.Append(" = ").AppendPropertyText(Property, UhtPropertyTextType.GenericFunctionArgOrRetVal)
+					.Append(property.SourceName)
+					.Append(" = ").AppendPropertyText(property, UhtPropertyTextType.GenericFunctionArgOrRetVal)
 					.Append("();");
 			}
-			return Builder;
+			return builder;
 		}
 
-		public static StringBuilder AppendGenericObjectDeclarationFromContext(this StringBuilder Builder, UhtClass Class)
+		public static StringBuilder AppendGenericObjectDeclarationFromContext(this StringBuilder builder, UhtClass _ /* classObj */)
 		{
-			return Builder.Append("UObject* Obj = (UObject*)InScriptContext;");
+			return builder.Append("UObject* Obj = (UObject*)InScriptContext;");
 		}
 
-		public static StringBuilder AppendGenericReturnValueHandler(this StringBuilder Builder, UhtClass Class, UhtProperty? ReturnValue)
+		public static StringBuilder AppendGenericReturnValueHandler(this StringBuilder builder, UhtClass _1 /* classObj */, UhtProperty? _2 /* returnValue */)
 		{
-			return Builder.Append("return 0;");
+			return builder.Append("return 0;");
 		}
 	}
 
 	internal class GenericScriptCodeGenerator : ScriptCodeGeneratorBase
 	{
-		public GenericScriptCodeGenerator(IUhtExportFactory Factory)
-			: base(Factory)
+		public GenericScriptCodeGenerator(IUhtExportFactory factory)
+			: base(factory)
 		{
 		}
 
-		protected override bool CanExportClass(UhtClass Class)
+		protected override bool CanExportClass(UhtClass classObj)
 		{
-			if (!base.CanExportClass(Class))
+			if (!base.CanExportClass(classObj))
 			{
 				return false;
 			}
 
-			foreach (UhtType Child in Class.Children)
+			foreach (UhtType child in classObj.Children)
 			{
-				if (Child is UhtFunction Function)
+				if (child is UhtFunction function)
 				{
-					if (CanExportFunction(Class, Function))
+					if (CanExportFunction(classObj, function))
 					{
 						return true;
 					}
 				}
-				else if (Child is UhtProperty Property)
+				else if (child is UhtProperty property)
 				{
-					if (CanExportProperty(Class, Property))
+					if (CanExportProperty(classObj, property))
 					{
 						return true;
 					}
@@ -81,122 +81,122 @@ namespace ScriptGeneratorUbtPlugin
 			return false;
 		}
 
-		protected override void ExportClass(StringBuilder Builder, UhtClass Class)
+		protected override void ExportClass(StringBuilder builder, UhtClass classObj)
 		{
-			Builder.Append("#pragma once\r\n\r\n");
+			builder.Append("#pragma once\r\n\r\n");
 
-			List<UhtFunction> Functions = Class.Children.Where(x => x is UhtFunction).Cast<UhtFunction>().Reverse().ToList();
+			List<UhtFunction> functions = classObj.Children.Where(x => x is UhtFunction).Cast<UhtFunction>().Reverse().ToList();
 
 			//ETSTODO - Functions are reversed in the engine
-			foreach (UhtFunction Function in Class.Functions.Reverse())
+			foreach (UhtFunction function in classObj.Functions.Reverse())
 			{
-				if (CanExportFunction(Class, Function))
+				if (CanExportFunction(classObj, function))
 				{
-					ExportFunction(Builder, Class, Function);
+					ExportFunction(builder, classObj, function);
 				}
 			}
-			//foreach (UhtType Child in Class.Children)
+			//foreach (UhtType child in Class.Children)
 			//{
-			//	if (Child is UhtFunction Function && CanExportFunction(Class, Function))
+			//	if (child is UhtFunction function && CanExportFunction(classObj, Function))
 			//	{
-			//		ExportFunction(Builder, Class, Function);
+			//		ExportFunction(Builder, classObj, function);
 			//	}
 			//}
 
-			foreach (UhtType Child in Class.Children)
+			foreach (UhtType child in classObj.Children)
 			{
-				if (Child is UhtProperty Property && CanExportProperty(Class, Property))
+				if (child is UhtProperty property && CanExportProperty(classObj, property))
 				{
-					ExportProperty(Builder, Class, Property);
+					ExportProperty(builder, classObj, property);
 				}
 			}
 
-			if (!Class.ClassFlags.HasAnyFlags(EClassFlags.Abstract))
+			if (!classObj.ClassFlags.HasAnyFlags(EClassFlags.Abstract))
 			{
-				Builder.AppendGenericWrapperFunctionDeclaration(Class, "New").Append("\r\n");
-				Builder.Append("{\r\n");
-				Builder.Append("\tUObject* Outer = NULL;\r\n");
-				Builder.Append("\tFName Name = FName(\"ScriptObject\");\r\n");
-				Builder.Append("\tUObject* Obj = NewObject<").Append(Class.SourceName).Append(">(Outer, Name);\r\n");
-				Builder.Append("\tif (Obj)\r\n\t{\r\n");
-				Builder.Append("\t\tFScriptObjectReferencer::Get().AddObjectReference(Obj);\r\n");
-				Builder.Append("\t\t// @todo: Register the object with the script context here\r\n");
-				Builder.Append("\t}\r\n\treturn 0;\r\n");
-				Builder.Append("}\r\n\r\n");
-
-				Builder.AppendGenericWrapperFunctionDeclaration(Class, "Destroy").Append("\r\n");
-				Builder.Append("{\r\n");
-				Builder.Append("\t").AppendGenericObjectDeclarationFromContext(Class).Append("\r\n");
-				Builder.Append("\tif (Obj)\r\n\t{\r\n");
-				Builder.Append("\t\tFScriptObjectReferencer::Get().RemoveObjectReference(Obj);\r\n");
-				Builder.Append("\t\t// @todo: Remove the object from the script context here if required\r\n");
-				Builder.Append("\t}\r\n\treturn 0;\r\n");
-				Builder.Append("}\r\n\r\n");
+				builder.AppendGenericWrapperFunctionDeclaration(classObj, "New").Append("\r\n");
+				builder.Append("{\r\n");
+				builder.Append("\tUObject* Outer = NULL;\r\n");
+				builder.Append("\tFName Name = FName(\"ScriptObject\");\r\n");
+				builder.Append("\tUObject* Obj = NewObject<").Append(classObj.SourceName).Append(">(Outer, Name);\r\n");
+				builder.Append("\tif (Obj)\r\n\t{\r\n");
+				builder.Append("\t\tFScriptObjectReferencer::Get().AddObjectReference(Obj);\r\n");
+				builder.Append("\t\t// @todo: Register the object with the script context here\r\n");
+				builder.Append("\t}\r\n\treturn 0;\r\n");
+				builder.Append("}\r\n\r\n");
+
+				builder.AppendGenericWrapperFunctionDeclaration(classObj, "Destroy").Append("\r\n");
+				builder.Append("{\r\n");
+				builder.Append('\t').AppendGenericObjectDeclarationFromContext(classObj).Append("\r\n");
+				builder.Append("\tif (Obj)\r\n\t{\r\n");
+				builder.Append("\t\tFScriptObjectReferencer::Get().RemoveObjectReference(Obj);\r\n");
+				builder.Append("\t\t// @todo: Remove the object from the script context here if required\r\n");
+				builder.Append("\t}\r\n\treturn 0;\r\n");
+				builder.Append("}\r\n\r\n");
 			}
 		}
 
-		protected void ExportFunction(StringBuilder Builder, UhtClass Class, UhtFunction Function)
+		protected void ExportFunction(StringBuilder builder, UhtClass classObj, UhtFunction function)
 		{
-			Builder.AppendGenericWrapperFunctionDeclaration(Class, Function.SourceName).Append("\r\n");
-			Builder.Append("{\r\n");
-			Builder.Append("\t").AppendGenericObjectDeclarationFromContext(Class).Append("\r\n");
-			AppendFunctionDispatch(Builder, Class, Function);
-			Builder.Append("\t").AppendGenericReturnValueHandler(Class, Function.ReturnProperty).Append("\r\n");
-			Builder.Append("}\r\n\r\n");
+			builder.AppendGenericWrapperFunctionDeclaration(classObj, function.SourceName).Append("\r\n");
+			builder.Append("{\r\n");
+			builder.Append('\t').AppendGenericObjectDeclarationFromContext(classObj).Append("\r\n");
+			AppendFunctionDispatch(builder, classObj, function);
+			builder.Append('\t').AppendGenericReturnValueHandler(classObj, function.ReturnProperty).Append("\r\n");
+			builder.Append("}\r\n\r\n");
 		}
 
-		protected void ExportProperty(StringBuilder Builder, UhtClass Class, UhtProperty Property)
+		protected static void ExportProperty(StringBuilder builder, UhtClass classObj, UhtProperty property)
 		{
 			// Getter
-			Builder.AppendGenericWrapperFunctionDeclaration(Class, $"Get_{Property.SourceName}").Append("\r\n");
-			Builder.Append("{\r\n");
-			Builder.Append('\t').AppendGenericObjectDeclarationFromContext(Class).Append("\r\n");
-			Builder.Append("\tstatic FProperty* Property = FindScriptPropertyHelper(").Append(Class.SourceName).Append("::StaticClass(), TEXT(\"").Append(Property.SourceName).Append("\"));\r\n");
-			Builder.Append('\t').AppendGenericFunctionParamDeclaration(Class, Property).Append("\r\n");
-			Builder.Append("\tProperty->CopyCompleteValue(&").Append(Property.SourceName).Append(", Property->ContainerPtrToValuePtr<void>(Obj));\r\n");
-			Builder.Append("\t// @todo: handle property value here\r\n");
-			Builder.Append("\treturn 0;\r\n");
-			Builder.Append("}\r\n\r\n");
+			builder.AppendGenericWrapperFunctionDeclaration(classObj, $"Get_{property.SourceName}").Append("\r\n");
+			builder.Append("{\r\n");
+			builder.Append('\t').AppendGenericObjectDeclarationFromContext(classObj).Append("\r\n");
+			builder.Append("\tstatic FProperty* Property = FindScriptPropertyHelper(").Append(classObj.SourceName).Append("::StaticClass(), TEXT(\"").Append(property.SourceName).Append("\"));\r\n");
+			builder.Append('\t').AppendGenericFunctionParamDeclaration(classObj, property).Append("\r\n");
+			builder.Append("\tProperty->CopyCompleteValue(&").Append(property.SourceName).Append(", Property->ContainerPtrToValuePtr<void>(Obj));\r\n");
+			builder.Append("\t// @todo: handle property value here\r\n");
+			builder.Append("\treturn 0;\r\n");
+			builder.Append("}\r\n\r\n");
 
 			// Setter
-			Builder.AppendGenericWrapperFunctionDeclaration(Class, $"Set_{Property.SourceName}").Append("\r\n");
-			Builder.Append("{\r\n");
-			Builder.Append('\t').AppendGenericObjectDeclarationFromContext(Class).Append("\r\n");
-			Builder.Append("\tstatic FProperty* Property = FindScriptPropertyHelper(").Append(Class.SourceName).Append("::StaticClass(), TEXT(\"").Append(Property.SourceName).Append("\"));\r\n");
-			Builder.Append('\t').AppendGenericFunctionParamDeclaration(Class, Property).Append("\r\n");
-			Builder.Append("\tProperty->CopyCompleteValue(Property->ContainerPtrToValuePtr<void>(Obj), &").Append(Property.SourceName).Append(");\r\n");
-			Builder.Append("\treturn 0;\r\n");
-			Builder.Append("}\r\n\r\n");
+			builder.AppendGenericWrapperFunctionDeclaration(classObj, $"Set_{property.SourceName}").Append("\r\n");
+			builder.Append("{\r\n");
+			builder.Append('\t').AppendGenericObjectDeclarationFromContext(classObj).Append("\r\n");
+			builder.Append("\tstatic FProperty* Property = FindScriptPropertyHelper(").Append(classObj.SourceName).Append("::StaticClass(), TEXT(\"").Append(property.SourceName).Append("\"));\r\n");
+			builder.Append('\t').AppendGenericFunctionParamDeclaration(classObj, property).Append("\r\n");
+			builder.Append("\tProperty->CopyCompleteValue(Property->ContainerPtrToValuePtr<void>(Obj), &").Append(property.SourceName).Append(");\r\n");
+			builder.Append("\treturn 0;\r\n");
+			builder.Append("}\r\n\r\n");
 		}
 
-		protected override void Finish(StringBuilder Builder, List<UhtClass> Classes)
+		protected override void Finish(StringBuilder builder, List<UhtClass> classes)
 		{
-			HashSet<UhtHeaderFile> UniqueHeaders = new HashSet<UhtHeaderFile>();
-			foreach (UhtClass Class in Classes)
+			HashSet<UhtHeaderFile> uniqueHeaders = new();
+			foreach (UhtClass classObj in classes)
 			{
-				UniqueHeaders.Add(Class.HeaderFile);
+				uniqueHeaders.Add(classObj.HeaderFile);
 			}
-			List<string> SortedHeaders = new List<string>();
-			foreach (UhtHeaderFile HeaderFile in UniqueHeaders)
+			List<string> sortedHeaders = new();
+			foreach (UhtHeaderFile headerFile in uniqueHeaders)
 			{
-				SortedHeaders.Add(HeaderFile.FilePath);
+				sortedHeaders.Add(headerFile.FilePath);
 			}
-			SortedHeaders.Sort(StringComparerUE.OrdinalIgnoreCase);
-			foreach (string FilePath in SortedHeaders)
+			sortedHeaders.Sort(StringComparerUE.OrdinalIgnoreCase);
+			foreach (string filePath in sortedHeaders)
 			{
-				string RelativePath = Path.GetRelativePath(this.Factory.PluginModule!.IncludeBase, FilePath).Replace("\\", "/");
-				Builder.Append("#include \"").Append(RelativePath).Append("\"\r\n");
+				string relativePath = Path.GetRelativePath(Factory.PluginModule!.IncludeBase, filePath).Replace("\\", "/");
+				builder.Append("#include \"").Append(relativePath).Append("\"\r\n");
 			}
 
-			SortedHeaders.Clear();
-			foreach (UhtClass Class in Classes)
+			sortedHeaders.Clear();
+			foreach (UhtClass classObj in classes)
 			{
-				SortedHeaders.Add($"{Class.EngineName}.script.h");
+				sortedHeaders.Add($"{classObj.EngineName}.script.h");
 			}
-			SortedHeaders.Sort(StringComparerUE.OrdinalIgnoreCase);
-			foreach (string FilePath in SortedHeaders)
+			sortedHeaders.Sort(StringComparerUE.OrdinalIgnoreCase);
+			foreach (string filePath in sortedHeaders)
 			{
-				Builder.Append("#include \"").Append(FilePath).Append("\"\r\n");
+				builder.Append("#include \"").Append(filePath).Append("\"\r\n");
 			}
 		}
 	}
diff --git a/Engine/Plugins/ScriptPlugin/Source/ScriptGeneratorUbtPlugin/LuaScriptCodeGenerator.cs b/Engine/Plugins/ScriptPlugin/Source/ScriptGeneratorUbtPlugin/LuaScriptCodeGenerator.cs
--- a/Engine/Plugins/ScriptPlugin/Source/ScriptGeneratorUbtPlugin/LuaScriptCodeGenerator.cs
+++ b/Engine/Plugins/ScriptPlugin/Source/ScriptGeneratorUbtPlugin/LuaScriptCodeGenerator.cs
@@ -5,7 +5,6 @@ using EpicGames.UHT.Types;
 using EpicGames.UHT.Utils;
 using System;
 using System.Collections.Generic;
-using System.Diagnostics;
 using System.IO;
 using System.Linq;
 using System.Text;
@@ -14,114 +13,114 @@ namespace ScriptGeneratorUbtPlugin
 {
 	static class LuaScriptCodeGeneratorStringBuilderExtensinos
 	{
-		public static StringBuilder AppendLuaWrapperFunctionDeclaration(this StringBuilder Builder, UhtClass Class, string FunctionName)
+		public static StringBuilder AppendLuaWrapperFunctionDeclaration(this StringBuilder builder, UhtClass classObj, string functionName)
 		{
-			return Builder.Append("int32 ").Append(Class.EngineName).Append('_').Append(FunctionName).Append("(lua_State* InScriptContext)");
+			return builder.Append("int32 ").Append(classObj.EngineName).Append('_').Append(functionName).Append("(lua_State* InScriptContext)");
 		}
 
-		public static StringBuilder AppendLuaObjectDeclarationFromContext(this StringBuilder Builder, UhtClass Class)
+		public static StringBuilder AppendLuaObjectDeclarationFromContext(this StringBuilder builder, UhtClass classObj)
 		{
-			return Builder.Append("UObject* Obj = (").Append(Class.SourceName).Append("*)lua_touserdata(InScriptContext, 1);");
+			return builder.Append("UObject* Obj = (").Append(classObj.SourceName).Append("*)lua_touserdata(InScriptContext, 1);");
 		}
 
-		public static StringBuilder AppendLuaReturnValueHandler(this StringBuilder Builder, UhtClass Class, UhtProperty? ReturnValue, string? ReturnValueName)
+		public static StringBuilder AppendLuaReturnValueHandler(this StringBuilder builder, UhtClass _ /* classObj */, UhtProperty? returnValue, string? returnValueName)
 		{
-			if (ReturnValue != null)
+			if (returnValue != null)
 			{
-				if (ReturnValue is UhtIntProperty)
+				if (returnValue is UhtIntProperty)
 				{
-					Builder.Append("lua_pushinteger(InScriptContext, ").Append(ReturnValueName).Append(");\r\n");
+					builder.Append("lua_pushinteger(InScriptContext, ").Append(returnValueName).Append(");\r\n");
 				}
-				else if (ReturnValue is UhtFloatProperty)
+				else if (returnValue is UhtFloatProperty)
 				{
-					Builder.Append("lua_pushnumber(InScriptContext, ").Append(ReturnValueName).Append(");\r\n");
+					builder.Append("lua_pushnumber(InScriptContext, ").Append(returnValueName).Append(");\r\n");
 				}
-				else if (ReturnValue is UhtStrProperty)
+				else if (returnValue is UhtStrProperty)
 				{
-					Builder.Append("lua_pushstring(InScriptContext, TCHAR_TO_ANSI(*").Append(ReturnValueName).Append("));\r\n");
+					builder.Append("lua_pushstring(InScriptContext, TCHAR_TO_ANSI(*").Append(returnValueName).Append("));\r\n");
 				}
-				else if (ReturnValue is UhtNameProperty)
+				else if (returnValue is UhtNameProperty)
 				{
-					Builder.Append("lua_pushstring(InScriptContext, TCHAR_TO_ANSI(*").Append(ReturnValueName).Append(".ToString()));\r\n");
+					builder.Append("lua_pushstring(InScriptContext, TCHAR_TO_ANSI(*").Append(returnValueName).Append(".ToString()));\r\n");
 				}
-				else if (ReturnValue is UhtBoolProperty)
+				else if (returnValue is UhtBoolProperty)
 				{
-					Builder.Append("lua_pushboolean(InScriptContext, ").Append(ReturnValueName).Append(");\r\n");
+					builder.Append("lua_pushboolean(InScriptContext, ").Append(returnValueName).Append(");\r\n");
 				}
-				else if (ReturnValue is UhtStructProperty StructProperty)
+				else if (returnValue is UhtStructProperty structProperty)
 				{
-					if (StructProperty.ScriptStruct.EngineName == "Vector2D")
+					if (structProperty.ScriptStruct.EngineName == "Vector2D")
 					{
-						Builder.Append("FLuaVector2D::Return(InScriptContext, ").Append(ReturnValueName).Append(");\r\n");
+						builder.Append("FLuaVector2D::Return(InScriptContext, ").Append(returnValueName).Append(");\r\n");
 					}
-					else if (StructProperty.ScriptStruct.EngineName == "Vector")
+					else if (structProperty.ScriptStruct.EngineName == "Vector")
 					{
-						Builder.Append("FLuaVector::Return(InScriptContext, ").Append(ReturnValueName).Append(");\r\n");
+						builder.Append("FLuaVector::Return(InScriptContext, ").Append(returnValueName).Append(");\r\n");
 					}
-					else if (StructProperty.ScriptStruct.EngineName == "Vector4")
+					else if (structProperty.ScriptStruct.EngineName == "Vector4")
 					{
-						Builder.Append("FLuaVector4::Return(InScriptContext, ").Append(ReturnValueName).Append(");\r\n");
+						builder.Append("FLuaVector4::Return(InScriptContext, ").Append(returnValueName).Append(");\r\n");
 					}
-					else if (StructProperty.ScriptStruct.EngineName == "Quat")
+					else if (structProperty.ScriptStruct.EngineName == "Quat")
 					{
-						Builder.Append("FLuaQuat::Return(InScriptContext, ").Append(ReturnValueName).Append(");\r\n");
+						builder.Append("FLuaQuat::Return(InScriptContext, ").Append(returnValueName).Append(");\r\n");
 					}
-					else if (StructProperty.ScriptStruct.EngineName == "LinearColor")
+					else if (structProperty.ScriptStruct.EngineName == "LinearColor")
 					{
-						Builder.Append("FLuaLinearColor::Return(InScriptContext, ").Append(ReturnValueName).Append(");\r\n");
+						builder.Append("FLuaLinearColor::Return(InScriptContext, ").Append(returnValueName).Append(");\r\n");
 					}
-					else if (StructProperty.ScriptStruct.EngineName == "Color")
+					else if (structProperty.ScriptStruct.EngineName == "Color")
 					{
-						Builder.Append("FLuaLinearColor::Return(InScriptContext, FLinearColor(").Append(ReturnValueName).Append("));\r\n");
+						builder.Append("FLuaLinearColor::Return(InScriptContext, FLinearColor(").Append(returnValueName).Append("));\r\n");
 					}
-					else if (StructProperty.ScriptStruct.EngineName == "Transform")
+					else if (structProperty.ScriptStruct.EngineName == "Transform")
 					{
-						Builder.Append("FLuaTransform::Return(InScriptContext, ").Append(ReturnValueName).Append(");\r\n");
+						builder.Append("FLuaTransform::Return(InScriptContext, ").Append(returnValueName).Append(");\r\n");
 					}
 					else
 					{
-						throw new UhtIceException($"Unsupported function return value struct type: {StructProperty.ScriptStruct.EngineName}");
+						throw new UhtIceException($"Unsupported function return value struct type: {structProperty.ScriptStruct.EngineName}");
 					}
 				}
-				else if (ReturnValue is UhtObjectPropertyBase)
+				else if (returnValue is UhtObjectPropertyBase)
 				{
-					Builder.Append("lua_pushlightuserdata(InScriptContext, ").Append(ReturnValueName).Append(");\r\n");
+					builder.Append("lua_pushlightuserdata(InScriptContext, ").Append(returnValueName).Append(");\r\n");
 				}
 				else
 				{
-					throw new UhtIceException($"Unsupported function return type: {ReturnValue.GetType().Name}");
+					throw new UhtIceException($"Unsupported function return type: {returnValue.GetType().Name}");
 				}
-				Builder.Append("\treturn 1;");
+				builder.Append("\treturn 1;");
 			}
 			else
 			{
-				Builder.Append("return 0;");
+				builder.Append("return 0;");
 			}
-			return Builder;
+			return builder;
 		}
 	}
 
 	internal class LuaScriptCodeGenerator : ScriptCodeGeneratorBase
 	{
-		public LuaScriptCodeGenerator(IUhtExportFactory Factory)
-			: base(Factory)
+		public LuaScriptCodeGenerator(IUhtExportFactory factory)
+			: base(factory)
 		{
 		}
 
-		protected override bool CanExportClass(UhtClass Class)
+		protected override bool CanExportClass(UhtClass classObj)
 		{
-			if (!base.CanExportClass(Class))
+			if (!base.CanExportClass(classObj))
 			{
 				return false;
 			}
 
-			for (UhtClass? Current = Class; Current != null; Current = Current.SuperClass)
+			for (UhtClass? current = classObj; current != null; current = current.SuperClass)
 			{
-				foreach (UhtType Child in Current.Children)
+				foreach (UhtType child in current.Children)
 				{
-					if (Child is UhtFunction Function)
+					if (child is UhtFunction function)
 					{
-						if (CanExportFunction(Class, Function))
+						if (CanExportFunction(classObj, function))
 						{
 							return true;
 						}
@@ -129,11 +128,11 @@ namespace ScriptGeneratorUbtPlugin
 				}
 			}
 
-			foreach (UhtType Child in Class.Children)
+			foreach (UhtType child in classObj.Children)
 			{
-				if (Child is UhtProperty Property)
+				if (child is UhtProperty property)
 				{
-					if (CanExportProperty(Class, Property))
+					if (CanExportProperty(classObj, property))
 					{
 						return true;
 					}
@@ -142,18 +141,18 @@ namespace ScriptGeneratorUbtPlugin
 			return false;
 		}
 
-		protected override bool CanExportFunction(UhtClass Class, UhtFunction Function)
+		protected override bool CanExportFunction(UhtClass classObj, UhtFunction function)
 		{
-			if (!base.CanExportFunction(Class, Function))
+			if (!base.CanExportFunction(classObj, function))
 			{
 				return false;
 			}
 
-			foreach (UhtType Child in Function.Children)
+			foreach (UhtType child in function.Children)
 			{
-				if (Child is UhtProperty Property)
+				if (child is UhtProperty property)
 				{
-					if (!IsPropertyTypeSupported(Property))
+					if (!IsPropertyTypeSupported(property))
 					{
 						return false;
 					}
@@ -162,326 +161,326 @@ namespace ScriptGeneratorUbtPlugin
 			return true;
 		}
 
-		protected override bool CanExportProperty(UhtClass Class, UhtProperty Property)
+		protected override bool CanExportProperty(UhtClass classObj, UhtProperty property)
 		{
-			return Property.PropertyFlags.HasAnyFlags(EPropertyFlags.Edit) && IsPropertyTypeSupported(Property);
+			return property.PropertyFlags.HasAnyFlags(EPropertyFlags.Edit) && IsPropertyTypeSupported(property);
 		}
 
-		private bool IsPropertyTypeSupported(UhtProperty Property)
+		private static bool IsPropertyTypeSupported(UhtProperty property)
 		{
-			if (Property is UhtStructProperty StructProperty)
+			if (property is UhtStructProperty structProperty)
 			{
 				return
-					StructProperty.ScriptStruct.EngineName == "Vector2D" ||
-					StructProperty.ScriptStruct.EngineName == "Vector" ||
-					StructProperty.ScriptStruct.EngineName == "Vector4" ||
-					StructProperty.ScriptStruct.EngineName == "Quat" ||
-					StructProperty.ScriptStruct.EngineName == "LinearColor" ||
-					StructProperty.ScriptStruct.EngineName == "Color" ||
-					StructProperty.ScriptStruct.EngineName == "Transform";
+					structProperty.ScriptStruct.EngineName == "Vector2D" ||
+					structProperty.ScriptStruct.EngineName == "Vector" ||
+					structProperty.ScriptStruct.EngineName == "Vector4" ||
+					structProperty.ScriptStruct.EngineName == "Quat" ||
+					structProperty.ScriptStruct.EngineName == "LinearColor" ||
+					structProperty.ScriptStruct.EngineName == "Color" ||
+					structProperty.ScriptStruct.EngineName == "Transform";
 			}
 			else if (
-				Property is UhtLazyObjectPtrProperty ||
-				Property is UhtSoftObjectProperty ||
-				Property is UhtSoftClassProperty ||
-				Property is UhtWeakObjectPtrProperty)
+				property is UhtLazyObjectPtrProperty ||
+				property is UhtSoftObjectProperty ||
+				property is UhtSoftClassProperty ||
+				property is UhtWeakObjectPtrProperty)
 			{
 				return false;
 			}
 			else if (
-				Property is UhtIntProperty ||
-				Property is UhtFloatProperty ||
-				Property is UhtStrProperty ||
-				Property is UhtNameProperty ||
-				Property is UhtBoolProperty ||
-				Property is UhtObjectPropertyBase)
+				property is UhtIntProperty ||
+				property is UhtFloatProperty ||
+				property is UhtStrProperty ||
+				property is UhtNameProperty ||
+				property is UhtBoolProperty ||
+				property is UhtObjectPropertyBase)
 			{
 				return true;
 			}
 			return false;
 		}
 
-		protected override void ExportClass(StringBuilder Builder, UhtClass Class)
+		protected override void ExportClass(StringBuilder builder, UhtClass classObj)
 		{
-			Builder.Append("#pragma once\r\n\r\n");
+			builder.Append("#pragma once\r\n\r\n");
 
-			List<UhtFunction> Functions = Class.Children.Where(x => x is UhtFunction).Cast<UhtFunction>().Reverse().ToList();
+			List<UhtFunction> functions = classObj.Children.Where(x => x is UhtFunction).Cast<UhtFunction>().Reverse().ToList();
 
 			//ETSTODO - Functions are reversed in the engine
-			for (UhtClass? Current = Class; Current != null; Current = Current.SuperClass)
+			for (UhtClass? current = classObj; current != null; current = current.SuperClass)
 			{
-				foreach (UhtFunction Function in Current.Functions.Reverse())
+				foreach (UhtFunction function in current.Functions.Reverse())
 				{
-					if (CanExportFunction(Class, Function))
+					if (CanExportFunction(classObj, function))
 					{
-						ExportFunction(Builder, Class, Function);
+						ExportFunction(builder, classObj, function);
 					}
 				}
 			}
 
-			for (UhtClass? Current = Class; Current != null; Current = Current.SuperClass)
+			for (UhtClass? current = classObj; current != null; current = current.SuperClass)
 			{
-				foreach (UhtType Child in Current.Children)
+				foreach (UhtType child in current.Children)
 				{
-					if (Child is UhtProperty Property && CanExportProperty(Class, Property))
+					if (child is UhtProperty property && CanExportProperty(classObj, property))
 					{
-						ExportProperty(Builder, Class, Property);
+						ExportProperty(builder, classObj, property);
 					}
 				}
 			}
 
-			if (!Class.ClassFlags.HasAnyFlags(EClassFlags.Abstract))
+			if (!classObj.ClassFlags.HasAnyFlags(EClassFlags.Abstract))
 			{
-				Builder.AppendLuaWrapperFunctionDeclaration(Class, "New").Append("\r\n");
-				Builder.Append("{\r\n");
-				Builder.Append("\tUObject* Outer = (UObject*)lua_touserdata(InScriptContext, 1);\r\n");
-				Builder.Append("\tFName Name = FName(luaL_checkstring(InScriptContext, 2));\r\n");
-				Builder.Append("\tUObject* Obj = NewObject<").Append(Class.SourceName).Append(">(Outer, Name);\r\n");
-				Builder.Append("\tif (Obj)\r\n\t{\r\n");
-				Builder.Append("\t\tFScriptObjectReferencer::Get().AddObjectReference(Obj);\r\n");
-				Builder.Append("\t}\r\n");
-				Builder.Append("\tlua_pushlightuserdata(InScriptContext, Obj);\r\n");
-				Builder.Append("\treturn 1;\r\n");
-				Builder.Append("}\r\n\r\n");
+				builder.AppendLuaWrapperFunctionDeclaration(classObj, "New").Append("\r\n");
+				builder.Append("{\r\n");
+				builder.Append("\tUObject* Outer = (UObject*)lua_touserdata(InScriptContext, 1);\r\n");
+				builder.Append("\tFName Name = FName(luaL_checkstring(InScriptContext, 2));\r\n");
+				builder.Append("\tUObject* Obj = NewObject<").Append(classObj.SourceName).Append(">(Outer, Name);\r\n");
+				builder.Append("\tif (Obj)\r\n\t{\r\n");
+				builder.Append("\t\tFScriptObjectReferencer::Get().AddObjectReference(Obj);\r\n");
+				builder.Append("\t}\r\n");
+				builder.Append("\tlua_pushlightuserdata(InScriptContext, Obj);\r\n");
+				builder.Append("\treturn 1;\r\n");
+				builder.Append("}\r\n\r\n");
 
-				Builder.AppendLuaWrapperFunctionDeclaration(Class, "Destroy").Append("\r\n");
-				Builder.Append("{\r\n");
-				Builder.Append("\t").AppendLuaObjectDeclarationFromContext(Class).Append("\r\n");
-				Builder.Append("\tif (Obj)\r\n\t{\r\n");
-				Builder.Append("\t\tFScriptObjectReferencer::Get().RemoveObjectReference(Obj);\r\n");
-				Builder.Append("\t}\r\n");
-				Builder.Append("\treturn 0;\r\n");
-				Builder.Append("}\r\n\r\n");
+				builder.AppendLuaWrapperFunctionDeclaration(classObj, "Destroy").Append("\r\n");
+				builder.Append("{\r\n");
+				builder.Append('\t').AppendLuaObjectDeclarationFromContext(classObj).Append("\r\n");
+				builder.Append("\tif (Obj)\r\n\t{\r\n");
+				builder.Append("\t\tFScriptObjectReferencer::Get().RemoveObjectReference(Obj);\r\n");
+				builder.Append("\t}\r\n");
+				builder.Append("\treturn 0;\r\n");
+				builder.Append("}\r\n\r\n");
 			}
 
 			// Class: Equivalent of StaticClass()
-			Builder.AppendLuaWrapperFunctionDeclaration(Class, "Class").Append("\r\n");
-			Builder.Append("{\r\n");
-			Builder.Append("\tUClass* Class = ").Append(Class.SourceName).Append("::StaticClass();\r\n");
-			Builder.Append("\tlua_pushlightuserdata(InScriptContext, Class);\r\n");
-			Builder.Append("\treturn 1;\r\n");
-			Builder.Append("}\r\n\r\n");
+			builder.AppendLuaWrapperFunctionDeclaration(classObj, "Class").Append("\r\n");
+			builder.Append("{\r\n");
+			builder.Append("\tUClass* Class = ").Append(classObj.SourceName).Append("::StaticClass();\r\n");
+			builder.Append("\tlua_pushlightuserdata(InScriptContext, Class);\r\n");
+			builder.Append("\treturn 1;\r\n");
+			builder.Append("}\r\n\r\n");
 
 			// Library
-			Builder.Append("static const luaL_Reg ").Append(Class.EngineName).Append("_Lib[] =\r\n");
-			Builder.Append("{\r\n");
-			if (!Class.ClassFlags.HasAnyFlags(EClassFlags.Abstract))
+			builder.Append("static const luaL_Reg ").Append(classObj.EngineName).Append("_Lib[] =\r\n");
+			builder.Append("{\r\n");
+			if (!classObj.ClassFlags.HasAnyFlags(EClassFlags.Abstract))
 			{
-				Builder.Append("\t{ \"New\", ").Append(Class.EngineName).Append("_New },\r\n");
-				Builder.Append("\t{ \"Destroy\", ").Append(Class.EngineName).Append("_Destroy },\r\n");
-				Builder.Append("\t{ \"Class\", ").Append(Class.EngineName).Append("_Class },\r\n");
+				builder.Append("\t{ \"New\", ").Append(classObj.EngineName).Append("_New },\r\n");
+				builder.Append("\t{ \"Destroy\", ").Append(classObj.EngineName).Append("_Destroy },\r\n");
+				builder.Append("\t{ \"Class\", ").Append(classObj.EngineName).Append("_Class },\r\n");
 			}
 
 			//ETSTODO - Functions are reversed in the engine
-			for (UhtClass? Current = Class; Current != null; Current = Current.SuperClass)
+			for (UhtClass? current = classObj; current != null; current = current.SuperClass)
 			{
-				foreach (UhtFunction Function in Current.Functions.Reverse())
+				foreach (UhtFunction function in current.Functions.Reverse())
 				{
-					if (CanExportFunction(Class, Function))
+					if (CanExportFunction(classObj, function))
 					{
-						Builder.Append("\t{ \"").Append(Function.SourceName).Append("\", ").Append(Class.EngineName).Append('_').Append(Function.SourceName).Append(" },\r\n");
+						builder.Append("\t{ \"").Append(function.SourceName).Append("\", ").Append(classObj.EngineName).Append('_').Append(function.SourceName).Append(" },\r\n");
 					}
 				}
 			}
 
-			for (UhtClass? Current = Class; Current != null; Current = Current.SuperClass)
+			for (UhtClass? current = classObj; current != null; current = current.SuperClass)
 			{
-				foreach (UhtType Child in Current.Children)
+				foreach (UhtType child in current.Children)
 				{
-					if (Child is UhtProperty Property && CanExportProperty(Class, Property))
+					if (child is UhtProperty property && CanExportProperty(classObj, property))
 					{
-						Builder.Append("\t{ \"Get_").Append(Property.SourceName).Append("\", ").Append(Class.EngineName).Append("_Get_").Append(Property.SourceName).Append(" },\r\n");
-						Builder.Append("\t{ \"Set_").Append(Property.SourceName).Append("\", ").Append(Class.EngineName).Append("_Set_").Append(Property.SourceName).Append(" },\r\n");
+						builder.Append("\t{ \"Get_").Append(property.SourceName).Append("\", ").Append(classObj.EngineName).Append("_Get_").Append(property.SourceName).Append(" },\r\n");
+						builder.Append("\t{ \"Set_").Append(property.SourceName).Append("\", ").Append(classObj.EngineName).Append("_Set_").Append(property.SourceName).Append(" },\r\n");
 					}
 				}
 			}
 
-			Builder.Append("\t{ NULL, NULL }\r\n");
-			Builder.Append("};\r\n\r\n");
+			builder.Append("\t{ NULL, NULL }\r\n");
+			builder.Append("};\r\n\r\n");
 		}
 
-		protected void ExportFunction(StringBuilder Builder, UhtClass Class, UhtFunction Function)
+		protected void ExportFunction(StringBuilder builder, UhtClass classObj, UhtFunction function)
 		{
-			UhtClass? FunctionSuper = null;
-			if (Function.Outer != null && Function.Outer != Class && Function.Outer is UhtClass OuterClass && base.CanExportClass(OuterClass))
+			UhtClass? functionSuper = null;
+			if (function.Outer != null && function.Outer != classObj && function.Outer is UhtClass outerClass && base.CanExportClass(outerClass))
 			{
-				FunctionSuper = OuterClass;
+				functionSuper = outerClass;
 			}
 
-			Builder.AppendLuaWrapperFunctionDeclaration(Class, Function.SourceName).Append("\r\n");
-			Builder.Append("{\r\n");
-			if (FunctionSuper == null)
+			builder.AppendLuaWrapperFunctionDeclaration(classObj, function.SourceName).Append("\r\n");
+			builder.Append("{\r\n");
+			if (functionSuper == null)
 			{
-				Builder.Append("\t").AppendLuaObjectDeclarationFromContext(Class).Append("\r\n");
-				AppendFunctionDispatch(Builder, Class, Function);
-				string ReturnValueName = Function.ReturnProperty != null ? $"Params.{Function.ReturnProperty.SourceName}" : String.Empty;
-				Builder.Append("\t").AppendLuaReturnValueHandler(Class, Function.ReturnProperty, ReturnValueName).Append("\r\n");
+				builder.Append('\t').AppendLuaObjectDeclarationFromContext(classObj).Append("\r\n");
+				AppendFunctionDispatch(builder, classObj, function);
+				string returnValueName = function.ReturnProperty != null ? $"Params.{function.ReturnProperty.SourceName}" : String.Empty;
+				builder.Append('\t').AppendLuaReturnValueHandler(classObj, function.ReturnProperty, returnValueName).Append("\r\n");
 			}
 			else
 			{
-				Builder.Append("\treturn ").Append(FunctionSuper.EngineName).Append('_').Append(Function.SourceName).Append("(InScriptContext);\r\n");
+				builder.Append("\treturn ").Append(functionSuper.EngineName).Append('_').Append(function.SourceName).Append("(InScriptContext);\r\n");
 			}
-			Builder.Append("}\r\n\r\n");
+			builder.Append("}\r\n\r\n");
 		}
 
-		protected void ExportProperty(StringBuilder Builder, UhtClass Class, UhtProperty Property)
+		protected void ExportProperty(StringBuilder builder, UhtClass classObj, UhtProperty property)
 		{
-			UhtClass? PropertySuper = null;
-			if (Property.Outer != null && Property.Outer != Class && Property.Outer is UhtClass OuterClass && base.CanExportClass(OuterClass))
+			UhtClass? propertySuper = null;
+			if (property.Outer != null && property.Outer != classObj && property.Outer is UhtClass outerClass && base.CanExportClass(outerClass))
 			{
-				PropertySuper = OuterClass;
+				propertySuper = outerClass;
 			}
 
 			// Getter
-			Builder.AppendLuaWrapperFunctionDeclaration(Class, $"Get_{Property.SourceName}").Append("\r\n");
-			Builder.Append("{\r\n");
-			if (PropertySuper == null)
+			builder.AppendLuaWrapperFunctionDeclaration(classObj, $"Get_{property.SourceName}").Append("\r\n");
+			builder.Append("{\r\n");
+			if (propertySuper == null)
 			{
-				Builder.Append('\t').AppendLuaObjectDeclarationFromContext(Class).Append("\r\n");
-				Builder.Append("\tstatic FProperty* Property = FindScriptPropertyHelper(").Append(Class.SourceName).Append("::StaticClass(), TEXT(\"").Append(Property.SourceName).Append("\"));\r\n");
-				Builder.Append('\t').AppendPropertyText(Property, UhtPropertyTextType.GenericFunctionArgOrRetVal).Append(" PropertyValue;\r\n");
-				Builder.Append("\tProperty->CopyCompleteValue(&PropertyValue, Property->ContainerPtrToValuePtr<void>(Obj));\r\n");
-				Builder.Append('\t').AppendLuaReturnValueHandler(Class, Property, "PropertyValue").Append("\r\n");
+				builder.Append('\t').AppendLuaObjectDeclarationFromContext(classObj).Append("\r\n");
+				builder.Append("\tstatic FProperty* Property = FindScriptPropertyHelper(").Append(classObj.SourceName).Append("::StaticClass(), TEXT(\"").Append(property.SourceName).Append("\"));\r\n");
+				builder.Append('\t').AppendPropertyText(property, UhtPropertyTextType.GenericFunctionArgOrRetVal).Append(" PropertyValue;\r\n");
+				builder.Append("\tProperty->CopyCompleteValue(&PropertyValue, Property->ContainerPtrToValuePtr<void>(Obj));\r\n");
+				builder.Append('\t').AppendLuaReturnValueHandler(classObj, property, "PropertyValue").Append("\r\n");
 			}
 			else
 			{
-				Builder.Append("\treturn ").Append(PropertySuper.EngineName).Append('_').Append("Get_").Append(Property.SourceName).Append("(InScriptContext);\r\n");
+				builder.Append("\treturn ").Append(propertySuper.EngineName).Append('_').Append("Get_").Append(property.SourceName).Append("(InScriptContext);\r\n");
 			}
-			Builder.Append("}\r\n\r\n");
+			builder.Append("}\r\n\r\n");
 
 			// Setter
-			Builder.AppendLuaWrapperFunctionDeclaration(Class, $"Set_{Property.SourceName}").Append("\r\n");
-			Builder.Append("{\r\n");
-			if (PropertySuper == null)
+			builder.AppendLuaWrapperFunctionDeclaration(classObj, $"Set_{property.SourceName}").Append("\r\n");
+			builder.Append("{\r\n");
+			if (propertySuper == null)
 			{
-				Builder.Append('\t').AppendLuaObjectDeclarationFromContext(Class).Append("\r\n");
-				Builder.Append("\tstatic FProperty* Property = FindScriptPropertyHelper(").Append(Class.SourceName).Append("::StaticClass(), TEXT(\"").Append(Property.SourceName).Append("\"));\r\n");
-				Builder.Append('\t').AppendPropertyText(Property, UhtPropertyTextType.GenericFunctionArgOrRetVal).Append(" PropertyValue = ");
-				AppendInitializeFunctionDispatchParam(Builder, Class, null, Property, 0).Append(";\r\n");
-				Builder.Append("\tProperty->CopyCompleteValue(Property->ContainerPtrToValuePtr<void>(Obj), &PropertyValue);\r\n");
-				Builder.Append("\treturn 0;\r\n");
+				builder.Append('\t').AppendLuaObjectDeclarationFromContext(classObj).Append("\r\n");
+				builder.Append("\tstatic FProperty* Property = FindScriptPropertyHelper(").Append(classObj.SourceName).Append("::StaticClass(), TEXT(\"").Append(property.SourceName).Append("\"));\r\n");
+				builder.Append('\t').AppendPropertyText(property, UhtPropertyTextType.GenericFunctionArgOrRetVal).Append(" PropertyValue = ");
+				AppendInitializeFunctionDispatchParam(builder, classObj, null, property, 0).Append(";\r\n");
+				builder.Append("\tProperty->CopyCompleteValue(Property->ContainerPtrToValuePtr<void>(Obj), &PropertyValue);\r\n");
+				builder.Append("\treturn 0;\r\n");
 			}
 			else
 			{
-				Builder.Append("\treturn ").Append(PropertySuper.EngineName).Append('_').Append("Set_").Append(Property.SourceName).Append("(InScriptContext);\r\n");
+				builder.Append("\treturn ").Append(propertySuper.EngineName).Append('_').Append("Set_").Append(property.SourceName).Append("(InScriptContext);\r\n");
 			}
-			Builder.Append("}\r\n\r\n");
+			builder.Append("}\r\n\r\n");
 		}
 
-		protected override void Finish(StringBuilder Builder, List<UhtClass> Classes)
+		protected override void Finish(StringBuilder builder, List<UhtClass> classes)
 		{
-			HashSet<UhtHeaderFile> UniqueHeaders = new HashSet<UhtHeaderFile>();
-			foreach (UhtClass Class in Classes)
+			HashSet<UhtHeaderFile> uniqueHeaders = new();
+			foreach (UhtClass classObj in classes)
 			{
-				UniqueHeaders.Add(Class.HeaderFile);
+				uniqueHeaders.Add(classObj.HeaderFile);
 			}
-			List<string> SortedHeaders = new List<string>();
-			foreach (UhtHeaderFile HeaderFile in UniqueHeaders)
+			List<string> sortedHeaders = new();
+			foreach (UhtHeaderFile headerFile in uniqueHeaders)
 			{
-				SortedHeaders.Add(HeaderFile.FilePath);
+				sortedHeaders.Add(headerFile.FilePath);
 			}
-			SortedHeaders.Sort(StringComparerUE.OrdinalIgnoreCase);
-			foreach (string FilePath in SortedHeaders)
+			sortedHeaders.Sort(StringComparerUE.OrdinalIgnoreCase);
+			foreach (string filePath in sortedHeaders)
 			{
-				string RelativePath = Path.GetRelativePath(this.Factory.PluginModule!.IncludeBase, FilePath).Replace("\\", "/");
-				Builder.Append("#include \"").Append(RelativePath).Append("\"\r\n");
+				string relativePath = Path.GetRelativePath(Factory.PluginModule!.IncludeBase, filePath).Replace("\\", "/");
+				builder.Append("#include \"").Append(relativePath).Append("\"\r\n");
 			}
 
-			Classes.Sort((x, y) => StringComparerUE.OrdinalIgnoreCase.Compare(x.EngineName, y.EngineName));
-			SortedHeaders.Clear();
-			foreach (UhtClass Class in Classes)
+			classes.Sort((x, y) => StringComparerUE.OrdinalIgnoreCase.Compare(x.EngineName, y.EngineName));
+			sortedHeaders.Clear();
+			foreach (UhtClass classObj in classes)
 			{
-				Builder.Append("#include \"").Append(Class.EngineName).Append(".script.h\"\r\n");
+				builder.Append("#include \"").Append(classObj.EngineName).Append(".script.h\"\r\n");
 			}
-			Builder.Append("\r\n");
-			Builder.Append("void LuaRegisterExportedClasses(lua_State* InScriptContext)\r\n");
-			Builder.Append("{\r\n");
-			foreach(UhtClass Class in Classes)
+			builder.Append("\r\n");
+			builder.Append("void LuaRegisterExportedClasses(lua_State* InScriptContext)\r\n");
+			builder.Append("{\r\n");
+			foreach(UhtClass classObj in classes)
 			{
-				Builder.Append("\tFLuaUtils::RegisterLibrary(InScriptContext, ").Append(Class.EngineName).Append("_Lib, \"").Append(Class.EngineName).Append("\");\r\n");
+				builder.Append("\tFLuaUtils::RegisterLibrary(InScriptContext, ").Append(classObj.EngineName).Append("_Lib, \"").Append(classObj.EngineName).Append("\");\r\n");
 			}
-			Builder.Append("}\r\n\r\n");
+			builder.Append("}\r\n\r\n");
 		}
 
-		protected override StringBuilder AppendInitializeFunctionDispatchParam(StringBuilder Builder, UhtClass Class, UhtFunction? Function, UhtProperty Property, int PropertyIndex)
+		protected override StringBuilder AppendInitializeFunctionDispatchParam(StringBuilder builder, UhtClass classObj, UhtFunction? function, UhtProperty property, int propertyIndex)
 		{
-			if (!Property.PropertyFlags.HasAnyFlags(EPropertyFlags.ReturnParm))
+			if (!property.PropertyFlags.HasAnyFlags(EPropertyFlags.ReturnParm))
 			{
-				int ParamIndex = PropertyIndex + 2;
-				if (Property is UhtIntProperty)
+				int paramIndex = propertyIndex + 2;
+				if (property is UhtIntProperty)
 				{
-					Builder.Append("(luaL_checkint");
+					builder.Append("(luaL_checkint");
 				}
-				else if (Property is UhtFloatProperty)
+				else if (property is UhtFloatProperty)
 				{
-					Builder.Append("(float)(luaL_checknumber");
+					builder.Append("(float)(luaL_checknumber");
 				}
-				else if (Property is UhtStrProperty)
+				else if (property is UhtStrProperty)
 				{
-					Builder.Append("ANSI_TO_TCHAR(luaL_checkstring");
+					builder.Append("ANSI_TO_TCHAR(luaL_checkstring");
 				}
-				else if (Property is UhtNameProperty)
+				else if (property is UhtNameProperty)
 				{
-					Builder.Append("FName(luaL_checkstring");
+					builder.Append("FName(luaL_checkstring");
 				}
-				else if (Property is UhtBoolProperty)
+				else if (property is UhtBoolProperty)
 				{
-					Builder.Append("!!(lua_toboolean");
+					builder.Append("!!(lua_toboolean");
 				}
-				else if (Property is UhtStructProperty StructProperty)
+				else if (property is UhtStructProperty structProperty)
 				{
-					if (StructProperty.ScriptStruct.EngineName == "Vector2D")
+					if (structProperty.ScriptStruct.EngineName == "Vector2D")
 					{
-						Builder.Append("(FLuaVector2D::Get");
+						builder.Append("(FLuaVector2D::Get");
 					}
-					else if (StructProperty.ScriptStruct.EngineName == "Vector")
+					else if (structProperty.ScriptStruct.EngineName == "Vector")
 					{
-						Builder.Append("(FLuaVector::Get");
+						builder.Append("(FLuaVector::Get");
 					}
-					else if (StructProperty.ScriptStruct.EngineName == "Vector4")
+					else if (structProperty.ScriptStruct.EngineName == "Vector4")
 					{
-						Builder.Append("(FLuaVector4::Get");
+						builder.Append("(FLuaVector4::Get");
 					}
-					else if (StructProperty.ScriptStruct.EngineName == "Quat")
+					else if (structProperty.ScriptStruct.EngineName == "Quat")
 					{
-						Builder.Append("(FLuaQuat::Get");
+						builder.Append("(FLuaQuat::Get");
 					}
-					else if (StructProperty.ScriptStruct.EngineName == "LinearColor")
+					else if (structProperty.ScriptStruct.EngineName == "LinearColor")
 					{
-						Builder.Append("(FLuaLinearColor::Get");
+						builder.Append("(FLuaLinearColor::Get");
 					}
-					else if (StructProperty.ScriptStruct.EngineName == "Color")
+					else if (structProperty.ScriptStruct.EngineName == "Color")
 					{
-						Builder.Append("FColor(FLuaLinearColor::Get");
+						builder.Append("FColor(FLuaLinearColor::Get");
 					}
-					else if (StructProperty.ScriptStruct.EngineName == "Transform")
+					else if (structProperty.ScriptStruct.EngineName == "Transform")
 					{
-						Builder.Append("(FLuaTransform::Get");
+						builder.Append("(FLuaTransform::Get");
 					}
 					else
 					{
-						throw new UhtIceException($"Unsupported function param struct type: {StructProperty.ScriptStruct.EngineName}");
+						throw new UhtIceException($"Unsupported function param struct type: {structProperty.ScriptStruct.EngineName}");
 					}
 				}
-				else if (Property is UhtClassProperty)
+				else if (property is UhtClassProperty)
 				{
-					Builder.Append("(UClass*)(lua_touserdata");
+					builder.Append("(UClass*)(lua_touserdata");
 				}
-				else if (Property is UhtObjectPropertyBase)
+				else if (property is UhtObjectPropertyBase)
 				{
-					Builder.Append("(").AppendPropertyText(Property, UhtPropertyTextType.GenericFunctionArgOrRetVal).Append(")(lua_touserdata");
+					builder.Append('(').AppendPropertyText(property, UhtPropertyTextType.GenericFunctionArgOrRetVal).Append(")(lua_touserdata");
 				}
 				else
 				{
-					throw new UhtIceException($"Unsupported function param type: {Property.GetType().Name}");
+					throw new UhtIceException($"Unsupported function param type: {property.GetType().Name}");
 				}
-				Builder.Append("(InScriptContext, ").Append(ParamIndex).Append("))");
+				builder.Append("(InScriptContext, ").Append(paramIndex).Append("))");
 			}
 			else
 			{
-				base.AppendInitializeFunctionDispatchParam(Builder, Class, Function, Property, PropertyIndex);
+				base.AppendInitializeFunctionDispatchParam(builder, classObj, function, property, propertyIndex);
 			}
-			return Builder;
+			return builder;
 		}
 	}
 }
diff --git a/Engine/Plugins/ScriptPlugin/Source/ScriptGeneratorUbtPlugin/ScriptCodeGeneratorBase.cs b/Engine/Plugins/ScriptPlugin/Source/ScriptGeneratorUbtPlugin/ScriptCodeGeneratorBase.cs
--- a/Engine/Plugins/ScriptPlugin/Source/ScriptGeneratorUbtPlugin/ScriptCodeGeneratorBase.cs
+++ b/Engine/Plugins/ScriptPlugin/Source/ScriptGeneratorUbtPlugin/ScriptCodeGeneratorBase.cs
@@ -13,14 +13,14 @@ using UnrealBuildTool;
 
 namespace ScriptGeneratorUbtPlugin
 {
-	abstract internal class ScriptCodeGeneratorBase
+	internal abstract class ScriptCodeGeneratorBase
 	{
 		public readonly IUhtExportFactory Factory;
-		public UhtSession Session => this.Factory.Session;
+		public UhtSession Session => Factory.Session;
 
-		public ScriptCodeGeneratorBase(IUhtExportFactory Factory)
+		public ScriptCodeGeneratorBase(IUhtExportFactory factory)
 		{
-			this.Factory = Factory;
+			Factory = factory;
 		}
 
 		/// <summary>
@@ -28,98 +28,98 @@ namespace ScriptGeneratorUbtPlugin
 		/// </summary>
 		public void Generate()
 		{
-			DirectoryReference ConfigDirectory = DirectoryReference.Combine(Unreal.EngineDirectory, "Programs/UnrealBuildTool");
-			ConfigHierarchy Ini = ConfigCache.ReadHierarchy(ConfigHierarchyType.Engine, ConfigDirectory, BuildHostPlatform.Current.Platform);
-			Ini.GetArray("Plugins", "ScriptSupportedModules", out List<string>? SupportedScriptModules);
+			DirectoryReference configDirectory = DirectoryReference.Combine(Unreal.EngineDirectory, "Programs/UnrealBuildTool");
+			ConfigHierarchy ini = ConfigCache.ReadHierarchy(ConfigHierarchyType.Engine, configDirectory, BuildHostPlatform.Current.Platform);
+			ini.GetArray("Plugins", "ScriptSupportedModules", out List<string>? supportedScriptModules);
 
 			// Loop through the packages making sure they should be exported.  Queue the export of the classes
-			List<UhtClass> Classes = new();
-			List<Task?> Tasks = new();
-			foreach (UhtPackage Package in this.Session.Packages)
+			List<UhtClass> classes = new();
+			List<Task?> tasks = new();
+			foreach (UhtPackage package in Session.Packages)
 			{
-				if (Package.Module.ModuleType != UHTModuleType.EngineRuntime && Package.Module.ModuleType != UHTModuleType.GameRuntime)
+				if (package.Module.ModuleType != UHTModuleType.EngineRuntime && package.Module.ModuleType != UHTModuleType.GameRuntime)
 				{
 					continue;
 				}
 
-				if (SupportedScriptModules != null && !SupportedScriptModules.Any(x => string.Compare(x, Package.Module.Name, StringComparison.OrdinalIgnoreCase) == 0))
+				if (supportedScriptModules != null && !supportedScriptModules.Any(x => String.Compare(x, package.Module.Name, StringComparison.OrdinalIgnoreCase) == 0))
 				{
 					continue;
 				}
 
-				QueueClassExports(Package, Package, Classes, Tasks);
+				QueueClassExports(package, package, classes, tasks);
 			}
 
 			// Wait for all the classes to export
-			Task[]? WaitTasks = Tasks.Where(x => x != null).Cast<Task>().ToArray();
-			if (WaitTasks.Length > 0)
+			Task[]? waitTasks = tasks.Where(x => x != null).Cast<Task>().ToArray();
+			if (waitTasks.Length > 0)
 			{
-				Task.WaitAll(WaitTasks);
+				Task.WaitAll(waitTasks);
 			}
 
 			// Finish the export process
-			Finish(Classes);
+			Finish(classes);
 		}
 
 		/// <summary>
 		/// Collect the classes to be exported for the given package and type
 		/// </summary>
-		/// <param name="Package">Package being exported</param>
-		/// <param name="Type">Type to test for exporting</param>
-		/// <param name="Classes">Collection of exported classes</param>
-		/// <param name="Tasks">Collection of queued tasks</param>
-		private void QueueClassExports(UhtPackage Package, UhtType Type, List<UhtClass> Classes, List<Task?> Tasks)
+		/// <param name="package">Package being exported</param>
+		/// <param name="type">Type to test for exporting</param>
+		/// <param name="classes">Collection of exported classes</param>
+		/// <param name="tasks">Collection of queued tasks</param>
+		private void QueueClassExports(UhtPackage package, UhtType type, List<UhtClass> classes, List<Task?> tasks)
 		{
-			if (Type is UhtClass Class)
+			if (type is UhtClass classObj)
 			{
-				if (CanExportClass(Class))
+				if (CanExportClass(classObj))
 				{
-					Classes.Add(Class);
-					Tasks.Add(Factory.CreateTask((Factory) => { ExportClass(Class); }));
+					classes.Add(classObj);
+					tasks.Add(Factory.CreateTask((factory) => { ExportClass(classObj); }));
 				}
 			}
-			foreach (UhtType Child in Type.Children)
+			foreach (UhtType child in type.Children)
 			{
-				QueueClassExports(Package, Child, Classes, Tasks);
+				QueueClassExports(package, child, classes, tasks);
 			}
 		}
 
 		/// <summary>
 		/// Test to see if the given class should be exported
 		/// </summary>
-		/// <param name="Class">Class to test</param>
+		/// <param name="classObj">Class to test</param>
 		/// <returns>True if the class should be exported, false if not</returns>
-		protected virtual bool CanExportClass(UhtClass Class)
+		protected virtual bool CanExportClass(UhtClass classObj)
 		{
-			return Class.ClassFlags.HasAnyFlags(EClassFlags.RequiredAPI | EClassFlags.MinimalAPI); // Don't export classes that don't export DLL symbols
+			return classObj.ClassFlags.HasAnyFlags(EClassFlags.RequiredAPI | EClassFlags.MinimalAPI); // Don't export classes that don't export DLL symbols
 		}
 
 		/// <summary>
 		/// Test to see if the given function should be exported
 		/// </summary>
-		/// <param name="Class">Owning class of the function</param>
-		/// <param name="Function">Function to test</param>
+		/// <param name="classObj">Owning class of the function</param>
+		/// <param name="function">Function to test</param>
 		/// <returns>True if the function should be exported</returns>
-		protected virtual bool CanExportFunction(UhtClass Class, UhtFunction Function)
+		protected virtual bool CanExportFunction(UhtClass classObj, UhtFunction function)
 		{
 			// We don't support delegates and non-public functions
-			if (Function.FunctionFlags.HasAnyFlags(EFunctionFlags.Delegate))
+			if (function.FunctionFlags.HasAnyFlags(EFunctionFlags.Delegate))
 			{
 				return false;
 			}
 
 			// Reject if any of the parameter types is unsupported yet
-			foreach (UhtType Child in Function.Children)
+			foreach (UhtType child in function.Children)
 			{
-				if (Child is UhtArrayProperty ||
-					Child is UhtDelegateProperty ||
-					Child is UhtMulticastDelegateProperty ||
-					Child is UhtWeakObjectPtrProperty ||
-					Child is UhtInterfaceProperty)
+				if (child is UhtArrayProperty ||
+					child is UhtDelegateProperty ||
+					child is UhtMulticastDelegateProperty ||
+					child is UhtWeakObjectPtrProperty ||
+					child is UhtInterfaceProperty)
 				{
 					return false;
 				}
-				if (Child is UhtProperty Property && Property.IsStaticArray)
+				if (child is UhtProperty property && property.IsStaticArray)
 				{
 					return false;
 				}
@@ -130,32 +130,32 @@ namespace ScriptGeneratorUbtPlugin
 		/// <summary>
 		/// Test to see if the given property should be exported
 		/// </summary>
-		/// <param name="Class">Owning class of the property</param>
-		/// <param name="Property">Property to test</param>
+		/// <param name="classObj">Owning class of the property</param>
+		/// <param name="property">Property to test</param>
 		/// <returns>True if the property should be exported</returns>
-		protected virtual bool CanExportProperty(UhtClass Class, UhtProperty Property)
+		protected virtual bool CanExportProperty(UhtClass classObj, UhtProperty property)
 		{
 			// Property must be DLL exported
-			if (!Class.ClassFlags.HasAnyFlags(EClassFlags.RequiredAPI))
+			if (!classObj.ClassFlags.HasAnyFlags(EClassFlags.RequiredAPI))
 			{
 				return false;
 			}
 
 			// Only public, editable properties can be exported
-			if (Property.PropertyFlags.HasAnyFlags(EPropertyFlags.NativeAccessSpecifierPrivate | EPropertyFlags.NativeAccessSpecifierProtected) || 
-				!Property.PropertyFlags.HasAnyFlags(EPropertyFlags.Edit))
+			if (property.PropertyFlags.HasAnyFlags(EPropertyFlags.NativeAccessSpecifierPrivate | EPropertyFlags.NativeAccessSpecifierProtected) || 
+				!property.PropertyFlags.HasAnyFlags(EPropertyFlags.Edit))
 			{
 				return false;
 			}
 
 			// Reject if any of the parameter types is unsupported yet
-			if (Property.IsStaticArray ||
-				Property is UhtArrayProperty ||
-				Property is UhtDelegateProperty ||
-				Property is UhtMulticastDelegateProperty ||
-				Property is UhtWeakObjectPtrProperty ||
-				Property is UhtInterfaceProperty ||
-				Property is UhtStructProperty)
+			if (property.IsStaticArray ||
+				property is UhtArrayProperty ||
+				property is UhtDelegateProperty ||
+				property is UhtMulticastDelegateProperty ||
+				property is UhtWeakObjectPtrProperty ||
+				property is UhtInterfaceProperty ||
+				property is UhtStructProperty)
 			{
 				return false;
 			}
@@ -166,73 +166,73 @@ namespace ScriptGeneratorUbtPlugin
 		/// Export the given class
 		/// </summary>
 		/// <param name="Factory">Factory associated with the export</param>
-		/// <param name="Class">Class to export</param>
-		private void ExportClass(UhtClass Class)
+		/// <param name="classObj">Class to export</param>
+		private void ExportClass(UhtClass classObj)
 		{
-			using BorrowStringBuilder Borrower = new(StringBuilderCache.Big);
-			ExportClass(Borrower.StringBuilder, Class);
-			string FileName = this.Factory.MakePath(Class.EngineName, ".script.h");
-			this.Factory.CommitOutput(FileName, Borrower.StringBuilder);
+			using BorrowStringBuilder borrower = new(StringBuilderCache.Big);
+			ExportClass(borrower.StringBuilder, classObj);
+			string fileName = Factory.MakePath(classObj.EngineName, ".script.h");
+			Factory.CommitOutput(fileName, borrower.StringBuilder);
 		}
 
-		private void Finish(List<UhtClass> Classes)
+		private void Finish(List<UhtClass> classes)
 		{
-			using BorrowStringBuilder Borrower = new(StringBuilderCache.Big);
-			Finish(Borrower.StringBuilder, Classes);
-			string FileName = this.Factory.MakePath("GeneratedScriptLibraries", ".inl");
-			this.Factory.CommitOutput(FileName, Borrower.StringBuilder);
+			using BorrowStringBuilder borrower = new(StringBuilderCache.Big);
+			Finish(borrower.StringBuilder, classes);
+			string fileName = Factory.MakePath("GeneratedScriptLibraries", ".inl");
+			Factory.CommitOutput(fileName, borrower.StringBuilder);
 		}
 
-		abstract protected void ExportClass(StringBuilder Builder, UhtClass Class);
+		protected abstract void ExportClass(StringBuilder builder, UhtClass classObj);
 
-		abstract protected void Finish(StringBuilder Builder, List<UhtClass> Classes);
+		protected abstract void Finish(StringBuilder builder, List<UhtClass> classes);
 
-		virtual protected StringBuilder AppendInitializeFunctionDispatchParam(StringBuilder Builder, UhtClass Class, UhtFunction? Function, UhtProperty Property, int PropertyIndex)
+		protected virtual StringBuilder AppendInitializeFunctionDispatchParam(StringBuilder builder, UhtClass classObj, UhtFunction? function, UhtProperty property, int propertyIndex)
 		{
-			if (Property is UhtObjectPropertyBase)
+			if (property is UhtObjectPropertyBase)
 			{
-				Builder.Append("NULL");
+				builder.Append("NULL");
 			}
 			else
 			{
-				Builder.AppendPropertyText(Property, UhtPropertyTextType.GenericFunctionArgOrRetVal).Append("()");
+				builder.AppendPropertyText(property, UhtPropertyTextType.GenericFunctionArgOrRetVal).Append("()");
 			}
-			return Builder;
+			return builder;
 		}
 
-		virtual protected StringBuilder AppendFunctionDispatch(StringBuilder Builder, UhtClass Class, UhtFunction Function)
+		protected virtual StringBuilder AppendFunctionDispatch(StringBuilder builder, UhtClass classObj, UhtFunction function)
 		{
-			bool bHasParamsOrReturnValue = Function.Children.Count > 0;
-			if (bHasParamsOrReturnValue)
+			bool hasParamsOrReturnValue = function.Children.Count > 0;
+			if (hasParamsOrReturnValue)
 			{
-				Builder.Append("\tstruct FDispatchParams\r\n");
-				Builder.Append("\t{\r\n");
-				foreach (UhtProperty Property in Function.Children)
+				builder.Append("\tstruct FDispatchParams\r\n");
+				builder.Append("\t{\r\n");
+				foreach (UhtProperty property in function.Children)
 				{
-					Builder.Append("\t\t").AppendPropertyText(Property, UhtPropertyTextType.GenericFunctionArgOrRetVal).Append(' ').Append(Property.SourceName).Append(";\r\n");
+					builder.Append("\t\t").AppendPropertyText(property, UhtPropertyTextType.GenericFunctionArgOrRetVal).Append(' ').Append(property.SourceName).Append(";\r\n");
 				}
-				Builder.Append("\t} Params;\r\n");
-				int PropertyIndex = 0;
-				foreach (UhtProperty Property in Function.Children)
+				builder.Append("\t} Params;\r\n");
+				int propertyIndex = 0;
+				foreach (UhtProperty property in function.Children)
 				{
-					Builder.Append("\tParams.").Append(Property.SourceName).Append(" = ");
-					AppendInitializeFunctionDispatchParam(Builder, Class, Function, Property, PropertyIndex).Append(";\r\n");
-					PropertyIndex++;
+					builder.Append("\tParams.").Append(property.SourceName).Append(" = ");
+					AppendInitializeFunctionDispatchParam(builder, classObj, function, property, propertyIndex).Append(";\r\n");
+					propertyIndex++;
 				}
 			}
 
-			Builder.Append("\tstatic UFunction* Function = Obj->FindFunctionChecked(TEXT(\"").Append(Function.SourceName).Append("\"));\r\n");
+			builder.Append("\tstatic UFunction* Function = Obj->FindFunctionChecked(TEXT(\"").Append(function.SourceName).Append("\"));\r\n");
 
-			if (bHasParamsOrReturnValue)
+			if (hasParamsOrReturnValue)
 			{
-				Builder.Append("\tcheck(Function->ParmsSize == sizeof(FDispatchParams));\r\n");
-				Builder.Append("\tObj->ProcessEvent(Function, &Params);\r\n");
+				builder.Append("\tcheck(Function->ParmsSize == sizeof(FDispatchParams));\r\n");
+				builder.Append("\tObj->ProcessEvent(Function, &Params);\r\n");
 			}
 			else
 			{
-				Builder.Append("\tObj->ProcessEvent(Function, NULL);\r\n");
+				builder.Append("\tObj->ProcessEvent(Function, NULL);\r\n");
 			}
-			return Builder;
+			return builder;
 		}
 	}
 }
diff --git a/Engine/Source/Programs/Shared/EpicGames.UHT/.editorconfig b/Engine/Source/Programs/Shared/EpicGames.UHT/.editorconfig
--- a/Engine/Source/Programs/Shared/EpicGames.UHT/.editorconfig
+++ b/Engine/Source/Programs/Shared/EpicGames.UHT/.editorconfig
@@ -174,12 +174,13 @@ dotnet_diagnostic.IDE2003.severity = warning
 # csharp_style_allow_blank_line_after_colon_in_constructor_initializer_experimental
 dotnet_diagnostic.IDE2004.severity = warning
 
-dotnet_style_require_accessibility_modifiers=false
+dotnet_style_require_accessibility_modifiers = false
 
 # CSharp code style settings:
 [*.cs]
-csharp_prefer_simple_using_statement=false
-csharp_style_prefer_range_operator=false
+csharp_prefer_simple_using_statement = true
+csharp_style_prefer_range_operator = true
+csharp_style_implicit_object_creation_when_type_is_apparent = true
 
 # Newline settings
 csharp_new_line_before_open_brace = all
@@ -222,7 +223,7 @@ csharp_style_expression_bodied_accessors = true:error
 # Suggest more modern language features when available
 csharp_style_pattern_matching_over_is_with_cast_check = true:error
 csharp_style_pattern_matching_over_as_with_null_check = false:error
-csharp_style_inlined_variable_declaration = false:suggestion
+csharp_style_inlined_variable_declaration = true:suggestion
 csharp_style_throw_expression = true:error
 csharp_style_conditional_delegate_call = true:suggestion
 
@@ -271,10 +272,3 @@ dotnet_diagnostic.IDE0003.severity = none
 
 # IDE0079: Remove unnecessary supression - Different tool versions will complain about unused stuff.
 dotnet_diagnostic.IDE0079.severity = none
-
-
-#IDE0018: Variable declaration can be inlined
-#IDE0057: Slice can be simplified
-#IDE0063: 'using' statement can be simplified
-#IDE0090: 'new' expression can be simplified
-
diff --git a/Engine/Source/Programs/Shared/EpicGames.UHT/Exporters/CodeGen/UhtCodeGenerator.cs b/Engine/Source/Programs/Shared/EpicGames.UHT/Exporters/CodeGen/UhtCodeGenerator.cs
--- a/Engine/Source/Programs/Shared/EpicGames.UHT/Exporters/CodeGen/UhtCodeGenerator.cs
+++ b/Engine/Source/Programs/Shared/EpicGames.UHT/Exporters/CodeGen/UhtCodeGenerator.cs
@@ -20,7 +20,7 @@ namespace EpicGames.UHT.Exporters.CodeGen
 			HeaderFilters = new string[] { "*.generated.h" })]
 		public static void CodeGenerator(IUhtExportFactory factory)
 		{
-			UhtCodeGenerator generator = new UhtCodeGenerator(factory);
+			UhtCodeGenerator generator = new(factory);
 			generator.Generate();
 		}
 
@@ -67,7 +67,7 @@ namespace EpicGames.UHT.Exporters.CodeGen
 
 		private void Generate()
 		{
-			List<Task?> prereqs = new List<Task?>();
+			List<Task?> prereqs = new();
 
 			// Perform some startup initialization to compute things we need over and over again
 			if (this.Session.GoWide)
@@ -112,7 +112,7 @@ namespace EpicGames.UHT.Exporters.CodeGen
 			}
 
 			// Generate the files for the packages
-			List<Task?> generatedPackages = new List<Task?>(this.Session.PackageTypeCount);
+			List<Task?> generatedPackages = new(this.Session.PackageTypeCount);
 			foreach (UhtPackage package in this.Session.Packages)
 			{
 				UHTManifest.Module module = package.Module;
@@ -153,7 +153,7 @@ namespace EpicGames.UHT.Exporters.CodeGen
 			}
 
 			// Wait for all the packages to complete
-			List<Task> packageTasks = new List<Task>(this.Session.PackageTypeCount);
+			List<Task> packageTasks = new(this.Session.PackageTypeCount);
 			foreach (Task? output in generatedPackages)
 			{
 				if (output != null)
@@ -228,7 +228,7 @@ namespace EpicGames.UHT.Exporters.CodeGen
 		#region Information initialization
 		private void InitPackageInfo(UhtPackage package)
 		{
-			StringBuilder builder = new StringBuilder();
+			StringBuilder builder = new();
 
 			ref PackageInfo packageInfo = ref this._packageInfos[package.PackageTypeIndex];
 			packageInfo._strippedName = package.SourceName.Replace('/', '_');
@@ -280,7 +280,7 @@ namespace EpicGames.UHT.Exporters.CodeGen
 			{
 				while (filePath.StartsWith("../", StringComparison.Ordinal))
 				{
-					filePath = filePath.Substring(3);
+					filePath = filePath[3..];
 				}
 			}
 
@@ -374,7 +374,7 @@ namespace EpicGames.UHT.Exporters.CodeGen
 		/// <returns>Sorted list of the header files</returns>
 		private static List<UhtHeaderFile> GetSortedHeaderFiles(UhtPackage package)
 		{
-			List<UhtHeaderFile> sortedHeaders = new List<UhtHeaderFile>(package.Children.Count);
+			List<UhtHeaderFile> sortedHeaders = new(package.Children.Count);
 			foreach (UhtHeaderFile headerFile in package.Children)
 			{
 				if (headerFile.ShouldExport)
diff --git a/Engine/Source/Programs/Shared/EpicGames.UHT/Exporters/CodeGen/UhtHeaderCodeGenerator.cs b/Engine/Source/Programs/Shared/EpicGames.UHT/Exporters/CodeGen/UhtHeaderCodeGenerator.cs
--- a/Engine/Source/Programs/Shared/EpicGames.UHT/Exporters/CodeGen/UhtHeaderCodeGenerator.cs
+++ b/Engine/Source/Programs/Shared/EpicGames.UHT/Exporters/CodeGen/UhtHeaderCodeGenerator.cs
@@ -462,13 +462,13 @@ namespace EpicGames.UHT.Exporters.CodeGen
 			{
 				throw new UhtIceException("Marshal and call name must begin with 'delegate'");
 			}
-			return $"F{function.MarshalAndCallName.Substring(DelegatePrefix.Length)}_DelegateWrapper";
+			return $"F{function.MarshalAndCallName[DelegatePrefix.Length..]}_DelegateWrapper";
 		}
 
 		protected static string GetDelegateFunctionExtraParameter(UhtFunction function)
 		{
 			string returnType = function.FunctionFlags.HasAnyFlags(EFunctionFlags.MulticastDelegate) ? "FMulticastScriptDelegate" : "FScriptDelegate";
-			return $"const {returnType}& {function.SourceName.Substring(1)}";
+			return $"const {returnType}& {function.SourceName[1..]}";
 		}
 		#endregion
 
diff --git a/Engine/Source/Programs/Shared/EpicGames.UHT/Exporters/CodeGen/UhtHeaderCodeGeneratorCppFile.cs b/Engine/Source/Programs/Shared/EpicGames.UHT/Exporters/CodeGen/UhtHeaderCodeGeneratorCppFile.cs
--- a/Engine/Source/Programs/Shared/EpicGames.UHT/Exporters/CodeGen/UhtHeaderCodeGeneratorCppFile.cs
+++ b/Engine/Source/Programs/Shared/EpicGames.UHT/Exporters/CodeGen/UhtHeaderCodeGeneratorCppFile.cs
@@ -31,8 +31,8 @@ namespace EpicGames.UHT.Exporters.CodeGen
 		public void Generate(IUhtExportFactory factory)
 		{
 			ref UhtCodeGenerator.HeaderInfo headerInfo = ref this.HeaderInfos[this.HeaderFile.HeaderFileTypeIndex];
-			using (BorrowStringBuilder borrower = new BorrowStringBuilder(StringBuilderCache.Big))
 			{
+				using BorrowStringBuilder borrower = new(StringBuilderCache.Big);
 				StringBuilder builder = borrower.StringBuilder;
 
 				builder.Append(HeaderCopyright);
@@ -41,7 +41,7 @@ namespace EpicGames.UHT.Exporters.CodeGen
 
 				bool addedStructuredArchiveFromArchiveHeader = false;
 				bool addedArchiveUObjectFromStructuredArchiveHeader = false;
-				HashSet<UhtHeaderFile> addedIncludes = new HashSet<UhtHeaderFile>();
+				HashSet<UhtHeaderFile> addedIncludes = new();
 				addedIncludes.Add(this.HeaderFile);
 				foreach (UhtType type in this.HeaderFile.Children)
 				{
@@ -288,17 +288,17 @@ namespace EpicGames.UHT.Exporters.CodeGen
 
 				builder.Append(EnableDeprecationWarnings).Append("\r\n");
 
-				using (UhtBorrowBuffer borrowBuffer = new UhtBorrowBuffer(builder))
 				{
+					using UhtBorrowBuffer borrowBuffer = new(builder);
 					string cppFilePath = factory.MakePath(HeaderFile, ".gen.cpp");
-					StringView generatedBody = new StringView(borrowBuffer.Buffer.Memory);
+					StringView generatedBody = new(borrowBuffer.Buffer.Memory);
 					if (this.SaveExportedHeaders)
 					{
 						factory.CommitOutput(cppFilePath, generatedBody);
 					}
 
 					// Save the hash of the generated body 
-					this.HeaderInfos[this.HeaderFile.HeaderFileTypeIndex]._bodyHash = UhtHash.GenenerateTextHash(generatedBody.Span.Slice(generatedBodyStart, generatedBodyEnd - generatedBodyStart));
+					this.HeaderInfos[this.HeaderFile.HeaderFileTypeIndex]._bodyHash = UhtHash.GenenerateTextHash(generatedBody.Span[generatedBodyStart..generatedBodyEnd]);
 				}
 			}
 		}
@@ -317,7 +317,7 @@ namespace EpicGames.UHT.Exporters.CodeGen
 				enumDisplayNameFn = "nullptr";
 			}
 
-			using (UhtMacroBlockEmitter macroBlockEmitter = new UhtMacroBlockEmitter(builder, "WITH_EDITORONLY_DATA", enumObj.IsEditorOnly))
+			using (UhtMacroBlockEmitter macroBlockEmitter = new(builder, "WITH_EDITORONLY_DATA", enumObj.IsEditorOnly))
 			{
 
 				// If we don't have a zero 0 then we emit a static assert to verify we have one
@@ -378,8 +378,7 @@ namespace EpicGames.UHT.Exporters.CodeGen
 				int enumIndex = 0;
 				foreach (UhtEnumValue value in enumObj.EnumValues)
 				{
-					string? keyName;
-					if (!enumObj.MetaData.TryGetValue("OverrideName", enumIndex, out keyName))
+					if (!enumObj.MetaData.TryGetValue("OverrideName", enumIndex, out string? keyName))
 					{
 						keyName = value.Name.ToString();
 					}
@@ -415,8 +414,8 @@ namespace EpicGames.UHT.Exporters.CodeGen
 				builder.Append("\t\treturn ").Append(registrationName).Append(".InnerSingleton;\r\n");
 				builder.Append("\t}\r\n");
 
-				using (UhtBorrowBuffer borrowBuffer = new UhtBorrowBuffer(builder, hashCodeBlockStart, builder.Length - hashCodeBlockStart))
 				{
+					using UhtBorrowBuffer borrowBuffer = new(builder, hashCodeBlockStart, builder.Length - hashCodeBlockStart);
 					this.ObjectInfos[enumObj.ObjectTypeIndex]._hash = UhtHash.GenenerateTextHash(borrowBuffer.Buffer.Memory.Span);
 				}
 			}
@@ -607,7 +606,7 @@ namespace EpicGames.UHT.Exporters.CodeGen
 				builder.Append("\t}\r\n");
 			}
 
-			using (UhtBorrowBuffer borrowBuffer = new UhtBorrowBuffer(builder, hashCodeBlockStart, builder.Length - hashCodeBlockStart))
+			using (UhtBorrowBuffer borrowBuffer = new(builder, hashCodeBlockStart, builder.Length - hashCodeBlockStart))
 			{
 				this.ObjectInfos[scriptStruct.ObjectTypeIndex]._hash = UhtHash.GenenerateTextHash(borrowBuffer.Buffer.Memory.Span);
 			}
@@ -772,7 +771,7 @@ namespace EpicGames.UHT.Exporters.CodeGen
 
 			builder.AppendEndEditorOnlyGuard(function.FunctionFlags.HasAnyFlags(EFunctionFlags.EditorOnly));
 
-			using (UhtBorrowBuffer borrowBuffer = new UhtBorrowBuffer(builder, hashCodeBlockStart, builder.Length - hashCodeBlockStart))
+			using (UhtBorrowBuffer borrowBuffer = new(builder, hashCodeBlockStart, builder.Length - hashCodeBlockStart))
 			{
 				this.ObjectInfos[function.ObjectTypeIndex]._hash = UhtHash.GenenerateTextHash(borrowBuffer.Buffer.Memory.Span);
 			}
@@ -815,7 +814,7 @@ namespace EpicGames.UHT.Exporters.CodeGen
 
 		private static StringBuilder AppendExportProperties(StringBuilder builder, UhtScriptStruct scriptStruct, int tabs)
 		{
-			using (UhtMacroBlockEmitter emitter = new UhtMacroBlockEmitter(builder, "WITH_EDITORONLY_DATA"))
+			using (UhtMacroBlockEmitter emitter = new(builder, "WITH_EDITORONLY_DATA"))
 			{
 				foreach (UhtProperty property in scriptStruct.Properties)
 				{
@@ -833,9 +832,9 @@ namespace EpicGames.UHT.Exporters.CodeGen
 				return builder;
 			}
 
-			PropertyMemberContextImpl context = new PropertyMemberContextImpl(this.CodeGenerator, structObj, structSourceName, staticsName);
+			PropertyMemberContextImpl context = new(this.CodeGenerator, structObj, structSourceName, staticsName);
 
-			using (UhtMacroBlockEmitter emitter = new UhtMacroBlockEmitter(builder, "WITH_EDITORONLY_DATA"))
+			using (UhtMacroBlockEmitter emitter = new(builder, "WITH_EDITORONLY_DATA"))
 			{
 				bool hasAllEditorOnlyDataProperties = true;
 				foreach (UhtType type in structObj.Children)
@@ -872,9 +871,9 @@ namespace EpicGames.UHT.Exporters.CodeGen
 				return builder;
 			}
 
-			PropertyMemberContextImpl context = new PropertyMemberContextImpl(this.CodeGenerator, structObj, structSourceName, staticsName);
+			PropertyMemberContextImpl context = new(this.CodeGenerator, structObj, structSourceName, staticsName);
 
-			using (UhtMacroBlockEmitter emitter = new UhtMacroBlockEmitter(builder, "WITH_EDITORONLY_DATA"))
+			using (UhtMacroBlockEmitter emitter = new(builder, "WITH_EDITORONLY_DATA"))
 			{
 				bool hasAllEditorOnlyDataProperties = true;
 				foreach (UhtType type in structObj.Children)
@@ -911,7 +910,7 @@ namespace EpicGames.UHT.Exporters.CodeGen
 		private StringBuilder AppendClass(StringBuilder builder, UhtClass classObj)
 		{
 			// Collect the functions in reversed order
-			List<UhtFunction> reversedFunctions = new List<UhtFunction>(classObj.Functions.Where(x => IsRpcFunction(x) && ShouldExportFunction(x)));
+			List<UhtFunction> reversedFunctions = new(classObj.Functions.Where(x => IsRpcFunction(x) && ShouldExportFunction(x)));
 			reversedFunctions.Reverse();
 
 			// Check to see if we have any RPC functions for the editor
@@ -930,7 +929,7 @@ namespace EpicGames.UHT.Exporters.CodeGen
 			AppendPropertyAccessors(builder, classObj);
 
 			// Collect the callback function and sort by name to make the order stable
-			List<UhtFunction> callbackFunctions = new List<UhtFunction>(classObj.Functions.Where(x => x.FunctionFlags.HasAnyFlags(EFunctionFlags.Event) && x.SuperFunction == null));
+			List<UhtFunction> callbackFunctions = new(classObj.Functions.Where(x => x.FunctionFlags.HasAnyFlags(EFunctionFlags.Event) && x.SuperFunction == null));
 			callbackFunctions.Sort((x, y) => StringComparerUE.OrdinalIgnoreCase.Compare(x.EngineName, y.EngineName));
 
 			// VM -> C++ proxies (events and delegates).
@@ -975,11 +974,7 @@ namespace EpicGames.UHT.Exporters.CodeGen
 			}
 
 			// Scan the children to see what we have
-			bool hasProperties;
-			bool hasFunctions;
-			bool hasEditorFields;
-			bool allEditorFields;
-			GetFieldNotifyStats(classObj, out hasProperties, out hasFunctions, out hasEditorFields, out allEditorFields);
+			GetFieldNotifyStats(classObj, out bool hasProperties, out bool hasFunctions, out bool hasEditorFields, out bool allEditorFields);
 
 			if (allEditorFields)
 			{
@@ -1211,13 +1206,13 @@ namespace EpicGames.UHT.Exporters.CodeGen
 			string registrationName = $"Z_Registration_Info_UClass_{classObj.SourceName}";
 			string[]? sparseDataTypes = classObj.MetaData.GetStringArray(UhtNames.SparseClassDataTypes);
 
-			PropertyMemberContextImpl context = new PropertyMemberContextImpl(this.CodeGenerator, classObj, classObj.SourceName, staticsName);
+			PropertyMemberContextImpl context = new(this.CodeGenerator, classObj, classObj.SourceName, staticsName);
 
 			bool hasInterfaces = classObj.Bases.Any(x => x is UhtClass baseClass && baseClass.ClassFlags.HasAnyFlags(EClassFlags.Interface));
 
 			// Collect the functions to be exported
 			bool allEditorOnlyFunctions = true;
-			List<UhtFunction> sortedFunctions = new List<UhtFunction>();
+			List<UhtFunction> sortedFunctions = new();
 			foreach (UhtFunction function in classObj.Functions)
 			{
 				if (!function.FunctionFlags.HasAnyFlags(EFunctionFlags.Delegate))
@@ -1423,7 +1418,7 @@ namespace EpicGames.UHT.Exporters.CodeGen
 			}
 
 			// At this point, we can compute the hash... HOWEVER, in the old UHT extra data is appended to the hash block that isn't emitted to the actual output
-			using (BorrowStringBuilder hashBorrower = new BorrowStringBuilder(StringBuilderCache.Small))
+			using (BorrowStringBuilder hashBorrower = new(StringBuilderCache.Small))
 			{
 				StringBuilder hashBuilder = hashBorrower.StringBuilder;
 				hashBuilder.Append(builder, hashCodeBlockStart, builder.Length - hashCodeBlockStart);
@@ -1468,17 +1463,15 @@ namespace EpicGames.UHT.Exporters.CodeGen
 
 				if (this.Session.IncludeDebugOutput)
 				{
-					using (UhtBorrowBuffer borrowBuffer = new UhtBorrowBuffer(hashBuilder, saveLength, hashBuilder.Length - saveLength))
-					{
-						builder.Append("#if 0\r\n");
-						builder.Append(borrowBuffer.Buffer.Memory);
-						builder.Append("#endif\r\n");
-					}
+					using UhtBorrowBuffer borrowBuffer = new(hashBuilder, saveLength, hashBuilder.Length - saveLength);
+					builder.Append("#if 0\r\n");
+					builder.Append(borrowBuffer.Buffer.Memory);
+					builder.Append("#endif\r\n");
 				}
 
 				// Calculate generated class initialization code hash so that we know when it changes after hot-reload
-				using (UhtBorrowBuffer borrowBuffer = new UhtBorrowBuffer(hashBuilder))
 				{
+					using UhtBorrowBuffer borrowBuffer = new(hashBuilder);
 					this.ObjectInfos[classObj.ObjectTypeIndex]._hash = UhtHash.GenenerateTextHash(borrowBuffer.Buffer.Memory.Span);
 				}
 			}
@@ -1532,7 +1525,7 @@ namespace EpicGames.UHT.Exporters.CodeGen
 
 			bool allEditorOnly = true;
 
-			List<UhtFunction> sortedFunctions = new List<UhtFunction>();
+			List<UhtFunction> sortedFunctions = new();
 			foreach (UhtFunction function in classObj.Functions)
 			{
 				if (function.FunctionFlags.HasExactFlags(EFunctionFlags.Native | EFunctionFlags.NetRequest, EFunctionFlags.Native))
@@ -1545,8 +1538,8 @@ namespace EpicGames.UHT.Exporters.CodeGen
 
 			if (sortedFunctions.Count != 0)
 			{
-				using (UhtMacroBlockEmitter blockEmitter = new UhtMacroBlockEmitter(builder, "WITH_EDITOR", allEditorOnly))
 				{
+					using UhtMacroBlockEmitter blockEmitter = new(builder, "WITH_EDITOR", allEditorOnly);
 					builder.Append("\t\tUClass* Class = ").Append(classObj.SourceName).Append("::StaticClass();\r\n");
 					builder.Append("\t\tstatic const FNameNativePtrPair Funcs[] = {\r\n");
 
@@ -1580,8 +1573,8 @@ namespace EpicGames.UHT.Exporters.CodeGen
 			if (callbackFunctions.Count > 0)
 			{
 				bool isInterfaceClass = classObj.ClassFlags.HasAnyFlags(EClassFlags.Interface);
-				using (UhtMacroBlockEmitter blockEmitter = new UhtMacroBlockEmitter(builder, "WITH_EDITOR"))
 				{
+					using UhtMacroBlockEmitter blockEmitter = new(builder, "WITH_EDITOR");
 					foreach (UhtFunction function in callbackFunctions)
 					{
 						// Net response functions don't go into the VM
@@ -1741,7 +1734,7 @@ namespace EpicGames.UHT.Exporters.CodeGen
 
 		private static List<UhtScriptStruct> FindNoExportStructs(UhtStruct structObj)
 		{
-			List<UhtScriptStruct> outScriptStructs = new List<UhtScriptStruct>();
+			List<UhtScriptStruct> outScriptStructs = new();
 			FindNoExportStructsRecursive(outScriptStructs, structObj);
 			outScriptStructs.Reverse();
 			return outScriptStructs;
@@ -1750,21 +1743,21 @@ namespace EpicGames.UHT.Exporters.CodeGen
 		private class PropertyMemberContextImpl : IUhtPropertyMemberContext
 		{
 			private readonly UhtCodeGenerator _codeGenerator;
-			private readonly UhtStruct _outerStructInternal;
-			private readonly string _outerStructSourceNameInternal;
-			private readonly string _staticsNameInternal;
+			private readonly UhtStruct _outerStruct;
+			private readonly string _outerStructSourceName;
+			private readonly string _staticsName;
 
 			public PropertyMemberContextImpl(UhtCodeGenerator codeGenerator, UhtStruct outerStruct, string outerStructSourceName, string staticsName)
 			{
 				this._codeGenerator = codeGenerator;
-				this._outerStructInternal = outerStruct;
-				this._staticsNameInternal = staticsName;
-				this._outerStructSourceNameInternal = outerStructSourceName.Length == 0 ? outerStruct.SourceName : outerStructSourceName;
+				this._outerStruct = outerStruct;
+				this._staticsName = staticsName;
+				this._outerStructSourceName = outerStructSourceName.Length == 0 ? outerStruct.SourceName : outerStructSourceName;
 			}
 
-			public UhtStruct OuterStruct => this._outerStructInternal;
-			public string OuterStructSourceName => this._outerStructSourceNameInternal;
-			public string StaticsName => this._staticsNameInternal;
+			public UhtStruct OuterStruct => this._outerStruct;
+			public string OuterStructSourceName => this._outerStructSourceName;
+			public string StaticsName => this._staticsName;
 			public string NamePrefix => "NewProp_";
 			public string MetaDataSuffix => "_MetaData";
 
