Removing the C++ version of the script generator plugin.

diff --git a/Engine/Plugins/ScriptPlugin/ScriptPlugin.uplugin b/Engine/Plugins/ScriptPlugin/ScriptPlugin.uplugin
--- a/Engine/Plugins/ScriptPlugin/ScriptPlugin.uplugin
+++ b/Engine/Plugins/ScriptPlugin/ScriptPlugin.uplugin
@@ -14,16 +14,9 @@
 	"CanContainContent" : false,
 	"IsBetaVersion" : false,
 	"Installed" : false,
-	"CanBeUsedWithUnrealHeaderTool" : true,
-	"SupportedPrograms": [ "UnrealHeaderTool" ],
 	
 	"Modules" :
 	[
-		{
-			"Name" : "ScriptGeneratorPlugin",
-			"Type" : "Program",
-			"LoadingPhase" : "PostConfigInit"
-		},
 		{
 			"Name" : "ScriptPlugin",
 			"Type" : "Runtime",
diff --git a/Engine/Plugins/ScriptPlugin/Source/ScriptGeneratorPlugin/Private/GenericScriptCodeGenerator.cpp b/Engine/Plugins/ScriptPlugin/Source/ScriptGeneratorPlugin/Private/GenericScriptCodeGenerator.cpp
deleted file mode 100644
diff --git a/Engine/Plugins/ScriptPlugin/Source/ScriptGeneratorPlugin/Private/GenericScriptCodeGenerator.h b/Engine/Plugins/ScriptPlugin/Source/ScriptGeneratorPlugin/Private/GenericScriptCodeGenerator.h
deleted file mode 100644
diff --git a/Engine/Plugins/ScriptPlugin/Source/ScriptGeneratorPlugin/Private/LuaScriptCodeGenerator.cpp b/Engine/Plugins/ScriptPlugin/Source/ScriptGeneratorPlugin/Private/LuaScriptCodeGenerator.cpp
deleted file mode 100644
diff --git a/Engine/Plugins/ScriptPlugin/Source/ScriptGeneratorPlugin/Private/LuaScriptCodeGenerator.h b/Engine/Plugins/ScriptPlugin/Source/ScriptGeneratorPlugin/Private/LuaScriptCodeGenerator.h
deleted file mode 100644
diff --git a/Engine/Plugins/ScriptPlugin/Source/ScriptGeneratorPlugin/Private/ScriptCodeGeneratorBase.cpp b/Engine/Plugins/ScriptPlugin/Source/ScriptGeneratorPlugin/Private/ScriptCodeGeneratorBase.cpp
deleted file mode 100644
diff --git a/Engine/Plugins/ScriptPlugin/Source/ScriptGeneratorPlugin/Private/ScriptCodeGeneratorBase.h b/Engine/Plugins/ScriptPlugin/Source/ScriptGeneratorPlugin/Private/ScriptCodeGeneratorBase.h
deleted file mode 100644
diff --git a/Engine/Plugins/ScriptPlugin/Source/ScriptGeneratorPlugin/Private/ScriptGeneratorLog.h b/Engine/Plugins/ScriptPlugin/Source/ScriptGeneratorPlugin/Private/ScriptGeneratorLog.h
deleted file mode 100644
diff --git a/Engine/Plugins/ScriptPlugin/Source/ScriptGeneratorPlugin/Private/ScriptGeneratorPlugin.cpp b/Engine/Plugins/ScriptPlugin/Source/ScriptGeneratorPlugin/Private/ScriptGeneratorPlugin.cpp
deleted file mode 100644
diff --git a/Engine/Plugins/ScriptPlugin/Source/ScriptGeneratorPlugin/Public/IScriptGeneratorPlugin.h b/Engine/Plugins/ScriptPlugin/Source/ScriptGeneratorPlugin/Public/IScriptGeneratorPlugin.h
deleted file mode 100644
diff --git a/Engine/Plugins/ScriptPlugin/Source/ScriptGeneratorPlugin/ScriptGeneratorPlugin.Build.cs b/Engine/Plugins/ScriptPlugin/Source/ScriptGeneratorPlugin/ScriptGeneratorPlugin.Build.cs
deleted file mode 100644
