ScriptGeneratorUbtPlugin: Bump TargetFramework, fix compile error

#rnx

diff --git a/Engine/Plugins/ScriptPlugin/Source/ScriptGeneratorUbtPlugin/ScriptCodeGeneratorBase.cs b/Engine/Plugins/ScriptPlugin/Source/ScriptGeneratorUbtPlugin/ScriptCodeGeneratorBase.cs
--- a/Engine/Plugins/ScriptPlugin/Source/ScriptGeneratorUbtPlugin/ScriptCodeGeneratorBase.cs
+++ b/Engine/Plugins/ScriptPlugin/Source/ScriptGeneratorUbtPlugin/ScriptCodeGeneratorBase.cs
@@ -47,7 +47,7 @@ namespace ScriptGeneratorUbtPlugin
 					continue;
 				}
 
-				QueueClassExports(module.Package, module.Package, classes, tasks);
+				QueueClassExports(module.ScriptPackage, module.ScriptPackage, classes, tasks);
 			}
 
 			// Wait for all the classes to export
diff --git a/Engine/Plugins/ScriptPlugin/Source/ScriptGeneratorUbtPlugin/ScriptGeneratorUbtPlugin.ubtplugin.csproj b/Engine/Plugins/ScriptPlugin/Source/ScriptGeneratorUbtPlugin/ScriptGeneratorUbtPlugin.ubtplugin.csproj
--- a/Engine/Plugins/ScriptPlugin/Source/ScriptGeneratorUbtPlugin/ScriptGeneratorUbtPlugin.ubtplugin.csproj
+++ b/Engine/Plugins/ScriptPlugin/Source/ScriptGeneratorUbtPlugin/ScriptGeneratorUbtPlugin.ubtplugin.csproj
@@ -3,7 +3,7 @@
   <Import Project="$(EngineDir)\Source\Programs\Shared\UnrealEngine.csproj.props" />
   
   <PropertyGroup>
-    <TargetFramework>net6.0</TargetFramework>
+    <TargetFramework>net8.0</TargetFramework>
     <Configuration Condition=" '$(Configuration)' == '' ">Development</Configuration>
     <OutputType>Library</OutputType>
     <AppendTargetFrameworkToOutputPath>false</AppendTargetFrameworkToOutputPath>
