ScriptGeneratorUbtPlugin: Remove redundant properties from csproj

#rnx
#Synced-CL 51592834

diff --git a/Engine/Plugins/ScriptPlugin/Source/ScriptGeneratorUbtPlugin/ScriptGeneratorUbtPlugin.ubtplugin.csproj b/Engine/Plugins/ScriptPlugin/Source/ScriptGeneratorUbtPlugin/ScriptGeneratorUbtPlugin.ubtplugin.csproj
--- a/Engine/Plugins/ScriptPlugin/Source/ScriptGeneratorUbtPlugin/ScriptGeneratorUbtPlugin.ubtplugin.csproj
+++ b/Engine/Plugins/ScriptPlugin/Source/ScriptGeneratorUbtPlugin/ScriptGeneratorUbtPlugin.ubtplugin.csproj
@@ -1,38 +1,18 @@
 ﻿<Project Sdk="Microsoft.NET.Sdk">
   <Import Project="ScriptGeneratorUbtPlugin.ubtplugin.csproj.props" Condition="Exists('ScriptGeneratorUbtPlugin.ubtplugin.csproj.props')" />
   <Import Project="$(EngineDir)\Source\Programs\Shared\UnrealEngine.csproj.props" />
-  
+
   <PropertyGroup>
-    <TargetFramework>net8.0</TargetFramework>
-    <Configuration Condition=" '$(Configuration)' == '' ">Development</Configuration>
     <OutputType>Library</OutputType>
     <AppendTargetFrameworkToOutputPath>false</AppendTargetFrameworkToOutputPath>
     <GenerateAssemblyInfo>false</GenerateAssemblyInfo>
     <GenerateTargetFrameworkAttribute>false</GenerateTargetFrameworkAttribute>
-    <Configurations>Debug;Release;Development</Configurations>
     <RootNamespace>ScriptGeneratorUbtPlugin</RootNamespace>
     <AssemblyName>ScriptGeneratorUbtPlugin</AssemblyName>
     <OutputPath>..\..\..\..\Binaries\DotNET\UnrealBuildTool\Plugins\ScriptGeneratorUbtPlugin\</OutputPath>
     <TreatWarningsAsErrors>true</TreatWarningsAsErrors>
-    <DebugType>pdbonly</DebugType>
     <Nullable>enable</Nullable>
   </PropertyGroup>
-  <PropertyGroup Condition=" '$(Configuration)|$(Platform)' == 'Development|AnyCPU' ">
-    <DefineConstants>$(DefineConstants);TRACE</DefineConstants>
-    <Optimize>true</Optimize>
-  </PropertyGroup>
-  <PropertyGroup Condition=" '$(Configuration)|$(Platform)' == 'Release|AnyCPU' ">
-    <DefineConstants>$(DefineConstants);TRACE</DefineConstants>
-    <Optimize>true</Optimize>
-  </PropertyGroup>
-  <PropertyGroup Condition=" '$(Configuration)|$(Platform)' == 'Debug|AnyCPU' ">
-    <DefineConstants>$(DefineConstants);DEBUG;TRACE</DefineConstants>
-    <DebugSymbols>true</DebugSymbols>
-    <DebugType>full</DebugType>
-  </PropertyGroup>
-  <ItemGroup>
-    <PackageReference Include="Microsoft.CSharp" Version="4.7.0" />
-  </ItemGroup>
   <ItemGroup>
     <Reference Include="EpicGames.Build">
       <HintPath>$(EngineDir)\Binaries\DotNET\UnrealBuildTool\EpicGames.Build.dll</HintPath>
